Remove usage of unacceptable words

diff --git a/Engine/Plugins/Animation/ACLPlugin/Source/ACLPluginEditor/Private/ACLStatsDumpCommandlet.cpp b/Engine/Plugins/Animation/ACLPlugin/Source/ACLPluginEditor/Private/ACLStatsDumpCommandlet.cpp
--- a/Engine/Plugins/Animation/ACLPlugin/Source/ACLPluginEditor/Private/ACLStatsDumpCommandlet.cpp
+++ b/Engine/Plugins/Animation/ACLPlugin/Source/ACLPluginEditor/Private/ACLStatsDumpCommandlet.cpp
@@ -51,7 +51,7 @@ THIRD_PARTY_INCLUDES_END
 //		-error: Enables the exhaustive error dumping
 //		-auto: Uses automatic compression
 //		-acl: Uses ACL compression
-//		-MasterTolerance=<tolerance>: The error threshold used by automatic compression
+//		-ErrorTolerance=<tolerance>: The error threshold used by automatic compression
 //		-resume: If present, clip extraction or compression will continue where it left off
 //////////////////////////////////////////////////////////////////////////
 
@@ -1252,9 +1252,9 @@ int32 UACLStatsDumpCommandlet::Main(const FString& Params)
 		AutoCompressionSettings = FAnimationUtils::GetDefaultAnimationBoneCompressionSettings();
 		AutoCompressionSettings->bForceBelowThreshold = true;
 
-		if (ParamsMap.Contains(TEXT("MasterTolerance")))
+		if (ParamsMap.Contains(TEXT("ErrorTolerance")))
 		{
-			AutoCompressionSettings->ErrorThreshold = FCString::Atof(*ParamsMap[TEXT("MasterTolerance")]);
+			AutoCompressionSettings->ErrorThreshold = FCString::Atof(*ParamsMap[TEXT("ErrorTolerance")]);
 		}
 	}
 
