Fix PVS warnings

warning V1047: Lifetime of the lambda is greater than lifetime of the local variable
warning V752: Creating an object with placement new requires a buffer of large size

diff --git a/Engine/Plugins/Animation/ACLPlugin/Source/ACLPluginEditor/Private/ACLStatsDumpCommandlet.cpp b/Engine/Plugins/Animation/ACLPlugin/Source/ACLPluginEditor/Private/ACLStatsDumpCommandlet.cpp
--- a/Engine/Plugins/Animation/ACLPlugin/Source/ACLPluginEditor/Private/ACLStatsDumpCommandlet.cpp
+++ b/Engine/Plugins/Animation/ACLPlugin/Source/ACLPluginEditor/Private/ACLStatsDumpCommandlet.cpp
@@ -488,7 +488,7 @@ static void DumpClipDetailedError(const acl::track_array_qvvf& Tracks, UAnimSequ
 		acl::decompression_context<acl::debug_transform_decompression_settings> Context;
 		Context.initialize(*CompressedClipData);
 
-		Writer["error_per_frame_and_bone"] = [&](sjson::ArrayWriter& Writer)
+		Writer["error_per_frame_and_bone"] = [&](sjson::ArrayWriter& Writer)	//-V1047
 		{
 			SimpleTransformWriter PoseWriter(LossyLocalPoseTransforms);
 
@@ -551,7 +551,7 @@ static void DumpClipDetailedError(const acl::track_array_qvvf& Tracks, UAnimSequ
 		return;
 	}
 
-	Writer["error_per_frame_and_bone"] = [&](sjson::ArrayWriter& Writer)
+	Writer["error_per_frame_and_bone"] = [&](sjson::ArrayWriter& Writer)	//-V1047
 	{
 		for (uint32 SampleIndex = 0; SampleIndex < NumSamples; ++SampleIndex)
 		{
@@ -656,7 +656,7 @@ static void CompressWithUE4Auto(FCompressionContext& Context, bool PerformExhaus
 		const double UE4CompressionRatio = double(Context.UE4RawSize) / double(CompressedSize);
 		const double ACLCompressionRatio = double(Context.ACLRawSize) / double(CompressedSize);
 
-		Writer["ue4_auto"] = [&](sjson::ObjectWriter& Writer)
+		Writer["ue4_auto"] = [&](sjson::ObjectWriter& Writer)	//-V1047
 		{
 			Writer["algorithm_name"] = TCHAR_TO_ANSI(*Context.UE4Clip->BoneCompressionSettings->GetClass()->GetName());
 			Writer["codec_name"] = TCHAR_TO_ANSI(*GetCodecName(Context.UE4Clip->CompressedData.BoneCompressionCodec));
@@ -720,7 +720,7 @@ static void CompressWithACL(FCompressionContext& Context, bool PerformExhaustive
 		const double UE4CompressionRatio = double(Context.UE4RawSize) / double(CompressedSize);
 		const double ACLCompressionRatio = double(Context.ACLRawSize) / double(CompressedSize);
 
-		Writer["ue4_acl"] = [&](sjson::ObjectWriter& Writer)
+		Writer["ue4_acl"] = [&](sjson::ObjectWriter& Writer)	//-V1047
 		{
 			Writer["algorithm_name"] = TCHAR_TO_ANSI(*Context.UE4Clip->BoneCompressionSettings->GetClass()->GetName());
 			Writer["codec_name"] = TCHAR_TO_ANSI(*GetCodecName(Context.UE4Clip->CompressedData.BoneCompressionCodec));
@@ -809,7 +809,7 @@ static void CompressWithUE4KeyReduction(FCompressionContext& Context, bool Perfo
 		const double UE4CompressionRatio = double(Context.UE4RawSize) / double(CompressedSize);
 		const double ACLCompressionRatio = double(Context.ACLRawSize) / double(CompressedSize);
 
-		Writer["ue4_keyreduction"] = [&](sjson::ObjectWriter& Writer)
+		Writer["ue4_keyreduction"] = [&](sjson::ObjectWriter& Writer)	//-V1047
 		{
 			Writer["algorithm_name"] = TCHAR_TO_ANSI(*Context.UE4Clip->BoneCompressionSettings->GetClass()->GetName());
 			Writer["codec_name"] = TCHAR_TO_ANSI(*GetCodecName(Context.UE4Clip->CompressedData.BoneCompressionCodec));
@@ -839,7 +839,7 @@ static void CompressWithUE4KeyReduction(FCompressionContext& Context, bool Perfo
 			// Number of animated tracks (not constant/default)
 			int32 NumAnimatedTracks = 0;
 
-			Writer["dropped_track_keys"] = [&](sjson::ArrayWriter& Writer)
+			Writer["dropped_track_keys"] = [&](sjson::ArrayWriter& Writer)	//-V1047
 			{
 				const int32 NumTracks = Context.UE4Clip->GetDataModel()->GetNumBoneTracks();
 				const int32 NumSamples = Context.UE4Clip->GetNumberOfSampledKeys();
@@ -909,7 +909,7 @@ static void CompressWithUE4KeyReduction(FCompressionContext& Context, bool Perfo
 			Writer["total_num_animated_keys"] = TotalNumAnimatedKeys;
 			Writer["total_num_dropped_animated_keys"] = TotalNumDroppedAnimatedKeys;
 
-			Writer["dropped_pose_keys"] = [&](sjson::ArrayWriter& Writer)
+			Writer["dropped_pose_keys"] = [&](sjson::ArrayWriter& Writer)	//-V1047
 			{
 				const int32 NumTracks = Context.UE4Clip->GetDataModel()->GetNumBoneTracks();
 				const int32 NumSamples = Context.UE4Clip->GetNumberOfSampledKeys();
diff --git a/Engine/Plugins/Experimental/AnimNext/Source/AnimNext/Public/DecoratorBase/Decorator.h b/Engine/Plugins/Experimental/AnimNext/Source/AnimNext/Public/DecoratorBase/Decorator.h
--- a/Engine/Plugins/Experimental/AnimNext/Source/AnimNext/Public/DecoratorBase/Decorator.h
+++ b/Engine/Plugins/Experimental/AnimNext/Source/AnimNext/Public/DecoratorBase/Decorator.h
@@ -38,7 +38,7 @@
 		UE::AnimNext::FDecoratorMemoryLayout{ sizeof(DecoratorName), alignof(DecoratorName), sizeof(DecoratorName::FSharedData), alignof(DecoratorName::FSharedData), sizeof(DecoratorName::FInstanceData), alignof(DecoratorName::FInstanceData) }; \
 	void DecoratorName::ConstructInstance(UE::AnimNext::FExecutionContext& Context, UE::AnimNext::FWeakDecoratorPtr DecoratorPtr, const UE::AnimNext::FDecoratorDescription& DecoratorDesc, UE::AnimNext::FDecoratorInstance& DecoratorInstance) const \
 	{ \
-		FInstanceData* Data = new(&DecoratorInstance) FInstanceData(); \
+		FInstanceData* Data = new(&static_cast<FInstanceData&>(DecoratorInstance)) FInstanceData(); \
 		Data->Construct(Context, DecoratorPtr, static_cast<const FSharedData&>(DecoratorDesc)); \
 	} \
 	void DecoratorName::DestructInstance(UE::AnimNext::FExecutionContext& Context, UE::AnimNext::FWeakDecoratorPtr DecoratorPtr, const UE::AnimNext::FDecoratorDescription& DecoratorDesc, UE::AnimNext::FDecoratorInstance& DecoratorInstance) const \
