Update ACLPlugin uplugin file to match other Epic maintained plugins

diff --git a/Engine/Plugins/Animation/ACLPlugin/ACLPlugin.uplugin b/Engine/Plugins/Animation/ACLPlugin/ACLPlugin.uplugin
--- a/Engine/Plugins/Animation/ACLPlugin/ACLPlugin.uplugin
+++ b/Engine/Plugins/Animation/ACLPlugin/ACLPlugin.uplugin
@@ -5,11 +5,11 @@
 	"FriendlyName": "Animation Compression Library",
 	"Description": "Use the Animation Compression Library (ACL) to compress AnimSequences.",
 	"Category": "Animation",
-	"CreatedBy": "Nicholas Frechette",
-	"CreatedByURL": "https://nfrechette.github.io/",
-	"DocsURL": "[REDACTED_PROFILE_URL]/acl-ue4-plugin",
-	"MarketplaceURL": "com.epicgames.launcher://ue/marketplace/product/993b51e29e9a48d0b3fb449645b82cb6",
-	"SupportURL": "[REDACTED_PROFILE_URL]/acl-ue4-plugin",
+	"CreatedBy": "Epic Games, Inc.",
+	"CreatedByURL": "https://epicgames.com",
+	"DocsURL": "",
+	"MarketplaceURL": "",
+	"SupportURL": "",
 	"EnabledByDefault": false,
 	"CanContainContent": true,
 	"IsBetaVersion": false,
@@ -22,8 +22,8 @@
 		},
 		{
 			"Name": "ACLPluginEditor",
-			"Type" : "Editor",
-			"LoadingPhase" : "Default",
+			"Type": "Editor",
+			"LoadingPhase": "Default",
 			"PlatformDenyList": []
 		}
 	]
