Remove ACL plugin safety fallback codec as it causes more issues than it is worth

diff --git a/Engine/Plugins/Animation/ACLPlugin/Source/ACLPlugin/Classes/AnimBoneCompressionCodec_ACL.h b/Engine/Plugins/Animation/ACLPlugin/Source/ACLPlugin/Classes/AnimBoneCompressionCodec_ACL.h
--- a/Engine/Plugins/Animation/ACLPlugin/Source/ACLPlugin/Classes/AnimBoneCompressionCodec_ACL.h
+++ b/Engine/Plugins/Animation/ACLPlugin/Source/ACLPlugin/Classes/AnimBoneCompressionCodec_ACL.h
@@ -26,14 +26,7 @@ class UAnimBoneCompressionCodec_ACL : public UAnimBoneCompressionCodec_ACLBase
 {
 	GENERATED_UCLASS_BODY()
 
-	UPROPERTY(EditAnywhere, Category = "ACL Options", Instanced, meta = (EditInline))
-	TObjectPtr<UAnimBoneCompressionCodec> SafetyFallbackCodec;
-
 #if WITH_EDITORONLY_DATA
-	/** The error threshold after which we fallback on a safer encoding. */
-	UPROPERTY(EditAnywhere, Category = "ACL Options", meta = (ClampMin = "0"))
-	float SafetyFallbackThreshold;
-
 	/** The skeletal meshes used to estimate the skinning deformation during compression. */
 	UPROPERTY(EditAnywhere, Category = "ACL Options")
 	TArray<TObjectPtr<class USkeletalMesh>> OptimizationTargets;
@@ -56,12 +49,7 @@ class UAnimBoneCompressionCodec_ACL : public UAnimBoneCompressionCodec_ACLBase
 // @third party code - Epic Games End
 
 	//////////////////////////////////////////////////////////////////////////
-	// UObject implementation
-	virtual void PostInitProperties() override;
-	virtual void GetPreloadDependencies(TArray<UObject*>& OutDeps) override;
-
 	// UAnimBoneCompressionCodec implementation
-	virtual bool IsCodecValid() const override;
 // @third party code - Epic Games Begin
 	virtual void PopulateDDCKey(const UE::Anim::Compression::FAnimDDCKeyArgs& KeyArgs, FArchive& Ar) override;
 // @third party code - Epic Games End
@@ -71,11 +59,9 @@ class UAnimBoneCompressionCodec_ACL : public UAnimBoneCompressionCodec_ACLBase
 	virtual void GetCompressionSettings(acl::compression_settings& OutSettings, const ITargetPlatform* TargetPlatform) const override;
 // @third party code - Epic Games End
 	virtual TArray<class USkeletalMesh*> GetOptimizationTargets() const override { return OptimizationTargets; }
-	virtual ACLSafetyFallbackResult ExecuteSafetyFallback(acl::iallocator& Allocator, const acl::compression_settings& Settings, const acl::track_array_qvvf& RawClip, const acl::track_array_qvvf& BaseClip, const acl::compressed_tracks& CompressedClipData, const FCompressibleAnimData& CompressibleAnimData, FCompressibleAnimDataResult& OutResult);
 #endif
 
 	// UAnimBoneCompressionCodec implementation
-	virtual UAnimBoneCompressionCodec* GetCodec(const FString& DDCHandle);
 	virtual void DecompressPose(FAnimSequenceDecompressionContext& DecompContext, const BoneTrackArray& RotationPairs, const BoneTrackArray& TranslationPairs, const BoneTrackArray& ScalePairs, TArrayView<FTransform>& OutAtoms) const override;
 	virtual void DecompressBone(FAnimSequenceDecompressionContext& DecompContext, int32 TrackIndex, FTransform& OutAtom) const override;
 };
diff --git a/Engine/Plugins/Animation/ACLPlugin/Source/ACLPlugin/Private/AnimBoneCompressionCodec_ACL.cpp b/Engine/Plugins/Animation/ACLPlugin/Source/ACLPlugin/Private/AnimBoneCompressionCodec_ACL.cpp
--- a/Engine/Plugins/Animation/ACLPlugin/Source/ACLPlugin/Private/AnimBoneCompressionCodec_ACL.cpp
+++ b/Engine/Plugins/Animation/ACLPlugin/Source/ACLPlugin/Private/AnimBoneCompressionCodec_ACL.cpp
@@ -4,7 +4,6 @@
 #include "AnimBoneCompressionCodec_ACL.h"
 
 #if WITH_EDITORONLY_DATA
-#include "AnimBoneCompressionCodec_ACLSafe.h"
 #include "Engine/SkeletalMesh.h"
 #include "Rendering/SkeletalMeshModel.h"
 // @third party code - Epic Games Begin
@@ -29,7 +28,6 @@ UAnimBoneCompressionCodec_ACL::UAnimBoneCompressionCodec_ACL(const FObjectInitia
 	: Super(ObjectInitializer)
 {
 #if WITH_EDITORONLY_DATA
-	SafetyFallbackThreshold = 1.0f;			// 1cm, should be very rarely exceeded
 // @third party code - Epic Games Begin
 	bAllowFrameRemoval = false;
 	FrameRemovalThresholdType = ACLFrameRemovalThresholdType::ProportionOfFrames;
@@ -39,37 +37,6 @@ UAnimBoneCompressionCodec_ACL::UAnimBoneCompressionCodec_ACL(const FObjectInitia
 }
 
 #if WITH_EDITORONLY_DATA
-void UAnimBoneCompressionCodec_ACL::PostInitProperties()
-{
-	Super::PostInitProperties();
-
-	if (!IsTemplate())
-	{
-		// Ensure we are never null
-		SafetyFallbackCodec = NewObject<UAnimBoneCompressionCodec_ACLSafe>(this, NAME_None, RF_Public);
-	}
-}
-
-void UAnimBoneCompressionCodec_ACL::GetPreloadDependencies(TArray<UObject*>& OutDeps)
-{
-	Super::GetPreloadDependencies(OutDeps);
-
-	if (SafetyFallbackCodec != nullptr)
-	{
-		OutDeps.Add(SafetyFallbackCodec);
-	}
-}
-
-bool UAnimBoneCompressionCodec_ACL::IsCodecValid() const
-{
-	if (!Super::IsCodecValid())
-	{
-		return false;
-	}
-
-	return SafetyFallbackCodec != nullptr ? SafetyFallbackCodec->IsCodecValid() : true;
-}
-
 // @third party code - Epic Games Begin
 void UAnimBoneCompressionCodec_ACL::GetCompressionSettings(acl::compression_settings& OutSettings, const ITargetPlatform* TargetPlatform) const
 // @third party code - Epic Games End
@@ -100,29 +67,6 @@ void UAnimBoneCompressionCodec_ACL::GetCompressionSettings(acl::compression_sett
 // @third party code - Epic Games End
 }
 
-ACLSafetyFallbackResult UAnimBoneCompressionCodec_ACL::ExecuteSafetyFallback(acl::iallocator& Allocator, const acl::compression_settings& Settings, const acl::track_array_qvvf& RawClip, const acl::track_array_qvvf& BaseClip, const acl::compressed_tracks& CompressedClipData, const FCompressibleAnimData& CompressibleAnimData, FCompressibleAnimDataResult& OutResult)
-{
-	if (SafetyFallbackCodec != nullptr && SafetyFallbackThreshold > 0.0f)
-	{
-		checkSlow(CompressedClipData.is_valid(true).empty());
-
-		acl::decompression_context<UE4DefaultDBDecompressionSettings> Context;
-		Context.initialize(CompressedClipData);
-
-		const acl::track_error TrackError = acl::calculate_compression_error(Allocator, RawClip, Context, *Settings.error_metric, BaseClip);
-		if (TrackError.error >= SafetyFallbackThreshold)
-		{
-			UE_LOG(LogAnimationCompression, Verbose, TEXT("ACL Animation compressed size: %u bytes [%s]"), CompressedClipData.get_size(), *CompressibleAnimData.FullName);
-			UE_LOG(LogAnimationCompression, Warning, TEXT("ACL Animation error is too high, a safe fallback will be used instead: %.4f cm at %.4f on track %i [%s]"), TrackError.error, TrackError.sample_time, TrackError.index, *CompressibleAnimData.FullName);
-
-			// Just use the safety fallback
-			return SafetyFallbackCodec->Compress(CompressibleAnimData, OutResult) ? ACLSafetyFallbackResult::Success : ACLSafetyFallbackResult::Failure;
-		}
-	}
-
-	return ACLSafetyFallbackResult::Ignored;
-}
-
 // @third party code - Epic Games Begin
 void UAnimBoneCompressionCodec_ACL::PopulateDDCKey(const UE::Anim::Compression::FAnimDDCKeyArgs& KeyArgs, FArchive& Ar)
 {
@@ -135,7 +79,7 @@ void UAnimBoneCompressionCodec_ACL::PopulateDDCKey(const UE::Anim::Compression::
 	uint32 ForceRebuildVersion = 1;
 	uint32 SettingsHash = Settings.get_hash();
 
-	Ar	<< SafetyFallbackThreshold << ForceRebuildVersion << SettingsHash;
+	Ar	<< ForceRebuildVersion << SettingsHash;
 
 	for (USkeletalMesh* SkelMesh : OptimizationTargets)
 	{
@@ -145,29 +89,9 @@ void UAnimBoneCompressionCodec_ACL::PopulateDDCKey(const UE::Anim::Compression::
 			Ar << MeshModel->SkeletalMeshModelGUID;
 		}
 	}
-
-	if (SafetyFallbackCodec != nullptr)
-	{
-// @third party code - Epic Games Begin
-		SafetyFallbackCodec->PopulateDDCKey(KeyArgs, Ar);
-// @third party code - Epic Games End
-	}
 }
 #endif // WITH_EDITORONLY_DATA
 
-UAnimBoneCompressionCodec* UAnimBoneCompressionCodec_ACL::GetCodec(const FString& DDCHandle)
-{
-	const FString ThisHandle = GetCodecDDCHandle();
-	UAnimBoneCompressionCodec* CodecMatch = ThisHandle == DDCHandle ? this : nullptr;
-
-	if (CodecMatch == nullptr && SafetyFallbackCodec != nullptr)
-	{
-		CodecMatch = SafetyFallbackCodec->GetCodec(DDCHandle);
-	}
-
-	return CodecMatch;
-}
-
 void UAnimBoneCompressionCodec_ACL::DecompressPose(FAnimSequenceDecompressionContext& DecompContext, const BoneTrackArray& RotationPairs, const BoneTrackArray& TranslationPairs, const BoneTrackArray& ScalePairs, TArrayView<FTransform>& OutAtoms) const
 {
 	const FACLCompressedAnimData& AnimData = static_cast<const FACLCompressedAnimData&>(DecompContext.CompressedAnimData);
