Update ACL codec tooltips to clarify usage

diff --git a/Engine/Plugins/Animation/ACLPlugin/Source/ACLPlugin/Classes/AnimBoneCompressionCodec_ACL.h b/Engine/Plugins/Animation/ACLPlugin/Source/ACLPlugin/Classes/AnimBoneCompressionCodec_ACL.h
--- a/Engine/Plugins/Animation/ACLPlugin/Source/ACLPlugin/Classes/AnimBoneCompressionCodec_ACL.h
+++ b/Engine/Plugins/Animation/ACLPlugin/Source/ACLPlugin/Classes/AnimBoneCompressionCodec_ACL.h
@@ -20,7 +20,7 @@ enum class ACLFrameRemovalThresholdType: uint8
 class ITargetPlatform;
 // @third party code - Epic Games End
 
-/** The default codec implementation for ACL support with the minimal set of exposed features for ease of use. */
+/** Uses the open source Animation Compression Library with default settings suitable for general purpose animations. */
 UCLASS(MinimalAPI, config = Engine, meta = (DisplayName = "Anim Compress ACL"))
 class UAnimBoneCompressionCodec_ACL : public UAnimBoneCompressionCodec_ACLBase
 {
diff --git a/Engine/Plugins/Animation/ACLPlugin/Source/ACLPlugin/Classes/AnimBoneCompressionCodec_ACLCustom.h b/Engine/Plugins/Animation/ACLPlugin/Source/ACLPlugin/Classes/AnimBoneCompressionCodec_ACLCustom.h
--- a/Engine/Plugins/Animation/ACLPlugin/Source/ACLPlugin/Classes/AnimBoneCompressionCodec_ACLCustom.h
+++ b/Engine/Plugins/Animation/ACLPlugin/Source/ACLPlugin/Classes/AnimBoneCompressionCodec_ACLCustom.h
@@ -8,7 +8,7 @@
 #include "AnimBoneCompressionCodec_ACLBase.h"
 #include "AnimBoneCompressionCodec_ACLCustom.generated.h"
 
-/** The custom codec implementation for ACL support with all features supported. */
+/** Uses the open source Animation Compression Library with custom settings suitable for debugging purposes. */
 UCLASS(MinimalAPI, config = Engine, meta = (DisplayName = "Anim Compress ACL Custom"))
 class UAnimBoneCompressionCodec_ACLCustom : public UAnimBoneCompressionCodec_ACLBase
 {
diff --git a/Engine/Plugins/Animation/ACLPlugin/Source/ACLPlugin/Classes/AnimBoneCompressionCodec_ACLDatabase.h b/Engine/Plugins/Animation/ACLPlugin/Source/ACLPlugin/Classes/AnimBoneCompressionCodec_ACLDatabase.h
--- a/Engine/Plugins/Animation/ACLPlugin/Source/ACLPlugin/Classes/AnimBoneCompressionCodec_ACLDatabase.h
+++ b/Engine/Plugins/Animation/ACLPlugin/Source/ACLPlugin/Classes/AnimBoneCompressionCodec_ACLDatabase.h
@@ -51,7 +51,11 @@ struct FACLDatabaseCompressedAnimData final : public ICompressedAnimData
 	virtual bool IsValid() const override;
 };
 
-/** The default database codec implementation for ACL support with the minimal set of exposed features for ease of use. */
+/**
+  * Uses the open source Animation Compression Library with default settings and database support.
+  * The referenced database can be used to strip the least important keyframes on a per platform basis
+  * or they can be streamed in/out on demand through Blueprint or C++.
+  */
 UCLASS(MinimalAPI, config = Engine, meta = (DisplayName = "Anim Compress ACL Database"))
 class UAnimBoneCompressionCodec_ACLDatabase : public UAnimBoneCompressionCodec_ACLBase
 {
diff --git a/Engine/Plugins/Animation/ACLPlugin/Source/ACLPlugin/Classes/AnimBoneCompressionCodec_ACLSafe.h b/Engine/Plugins/Animation/ACLPlugin/Source/ACLPlugin/Classes/AnimBoneCompressionCodec_ACLSafe.h
--- a/Engine/Plugins/Animation/ACLPlugin/Source/ACLPlugin/Classes/AnimBoneCompressionCodec_ACLSafe.h
+++ b/Engine/Plugins/Animation/ACLPlugin/Source/ACLPlugin/Classes/AnimBoneCompressionCodec_ACLSafe.h
@@ -8,7 +8,7 @@
 #include "AnimBoneCompressionCodec_ACLBase.h"
 #include "AnimBoneCompressionCodec_ACLSafe.generated.h"
 
-/** The default codec implementation for ACL support with the minimal set of exposed features for ease of use. */
+/** Uses the open source Animation Compression Library with the safest and least destructive settings suitable when animations must be preserved with near raw fidelity. */
 UCLASS(MinimalAPI, config = Engine, meta = (DisplayName = "Anim Compress ACL Safe"))
 class UAnimBoneCompressionCodec_ACLSafe : public UAnimBoneCompressionCodec_ACLBase
 {
diff --git a/Engine/Plugins/Animation/ACLPlugin/Source/ACLPlugin/Classes/AnimCurveCompressionCodec_ACL.h b/Engine/Plugins/Animation/ACLPlugin/Source/ACLPlugin/Classes/AnimCurveCompressionCodec_ACL.h
--- a/Engine/Plugins/Animation/ACLPlugin/Source/ACLPlugin/Classes/AnimCurveCompressionCodec_ACL.h
+++ b/Engine/Plugins/Animation/ACLPlugin/Source/ACLPlugin/Classes/AnimCurveCompressionCodec_ACL.h
@@ -10,7 +10,7 @@
 #include "Animation/AnimCurveCompressionCodec.h"
 #include "AnimCurveCompressionCodec_ACL.generated.h"
 
-/** The default codec implementation for ACL curve compression support with the minimal set of exposed features for ease of use. */
+/** Uses the open source Animation Compression Library with default settings suitable for general purpose animation curves. */
 UCLASS(MinimalAPI, config = Engine, meta = (DisplayName = "ACL Curves"))
 class UAnimCurveCompressionCodec_ACL : public UAnimCurveCompressionCodec
 {
