Disable pch for ACL plugin as it compiles with cpp17 and is causing some issues with chained pch on clang platforms

#rnx

diff --git a/Engine/Plugins/Animation/ACLPlugin/Source/ACLPlugin/ACLPlugin.Build.cs b/Engine/Plugins/Animation/ACLPlugin/Source/ACLPlugin/ACLPlugin.Build.cs
--- a/Engine/Plugins/Animation/ACLPlugin/Source/ACLPlugin/ACLPlugin.Build.cs
+++ b/Engine/Plugins/Animation/ACLPlugin/Source/ACLPlugin/ACLPlugin.Build.cs
@@ -13,7 +13,8 @@ namespace UnrealBuildTool.Rules
 
             string ACLSDKDir = Path.GetFullPath(Path.Combine(ModuleDirectory, "../ThirdParty"));
 
-			PCHUsage = PCHUsageMode.UseExplicitOrSharedPCHs;
+			// Remove when this plugin can compile with cpp20
+			PCHUsage = PCHUsageMode.NoPCHs;
 
 			PublicIncludePaths.Add(Path.Combine(ACLSDKDir, "acl/includes"));
 			PublicIncludePaths.Add(Path.Combine(ACLSDKDir, "acl/external/rtm/includes"));
diff --git a/Engine/Plugins/Animation/ACLPlugin/Source/ACLPluginEditor/ACLPluginEditor.Build.cs b/Engine/Plugins/Animation/ACLPlugin/Source/ACLPluginEditor/ACLPluginEditor.Build.cs
--- a/Engine/Plugins/Animation/ACLPlugin/Source/ACLPluginEditor/ACLPluginEditor.Build.cs
+++ b/Engine/Plugins/Animation/ACLPlugin/Source/ACLPluginEditor/ACLPluginEditor.Build.cs
@@ -11,8 +11,8 @@ namespace UnrealBuildTool.Rules
 		{
 			CppStandard = CppStandardVersion.Cpp17;
 
-			PCHUsage = PCHUsageMode.UseExplicitOrSharedPCHs;
-			//OptimizeCode = CodeOptimization.Never;
+			// Remove when this plugin can compile with cpp20
+			PCHUsage = PCHUsageMode.NoPCHs;
 
 			PublicDependencyModuleNames.Add("ACLPlugin");
 			PublicDependencyModuleNames.Add("Core");
