Fix non-unity server build.

diff --git a/Engine/Plugins/Animation/ACLPlugin/Source/ACLPlugin/Private/ACLPluginModule.cpp b/Engine/Plugins/Animation/ACLPlugin/Source/ACLPlugin/Private/ACLPluginModule.cpp
--- a/Engine/Plugins/Animation/ACLPlugin/Source/ACLPlugin/Private/ACLPluginModule.cpp
+++ b/Engine/Plugins/Animation/ACLPlugin/Source/ACLPlugin/Private/ACLPluginModule.cpp
@@ -5,6 +5,7 @@
 #include "IACLPluginModule.h"
 #include "Misc/CoreDelegates.h"
 #include "Modules/ModuleManager.h"
+#include "ACLImpl.h"
 
 // Enable console commands only in development builds when logging is enabled
 #define WITH_ACL_CONSOLE_COMMANDS (!UE_BUILD_SHIPPING && !UE_BUILD_TEST && !NO_LOGGING)
