Remove experimental frame stripping from ACL

diff --git a/Engine/Plugins/Animation/ACLPlugin/Source/ACLPlugin/Classes/AnimBoneCompressionCodec_ACL.h b/Engine/Plugins/Animation/ACLPlugin/Source/ACLPlugin/Classes/AnimBoneCompressionCodec_ACL.h
--- a/Engine/Plugins/Animation/ACLPlugin/Source/ACLPlugin/Classes/AnimBoneCompressionCodec_ACL.h
+++ b/Engine/Plugins/Animation/ACLPlugin/Source/ACLPlugin/Classes/AnimBoneCompressionCodec_ACL.h
@@ -10,13 +10,6 @@
 #include "AnimBoneCompressionCodec_ACL.generated.h"
 
 // @third party code - Epic Games Begin
-UENUM()
-enum class ACLFrameRemovalThresholdType: uint8
-{
-	ProportionOfFrames UMETA(DisplayName = "Proportion"),
-	DistanceError UMETA(DisplayName = "Distance")
-};
-
 class ITargetPlatform;
 // @third party code - Epic Games End
 
@@ -31,23 +24,6 @@ class UAnimBoneCompressionCodec_ACL : public UAnimBoneCompressionCodec_ACLBase
 	UPROPERTY(EditAnywhere, Category = "ACL Options")
 	TArray<TObjectPtr<class USkeletalMesh>> OptimizationTargets;
 
-// @third party code - Epic Games Begin
-	/** Enable removal of the least important frames in the animation */
-	UPROPERTY(EditAnywhere, Category = "Frame Removal")
-	bool bAllowFrameRemoval;
-
-	UPROPERTY(EditAnywhere, Category = "Frame Removal", meta = (EditCondition = "bAllowFrameRemoval", EditConditionHides))
-	ACLFrameRemovalThresholdType FrameRemovalThresholdType;
-
-	/** Proportion of frames to remove (0.0 - 1.0)*/
-	UPROPERTY(EditAnywhere, Category = "Frame Removal", meta = (EditCondition = "bAllowFrameRemoval && FrameRemovalThresholdType == ACLFrameRemovalThresholdType::ProportionOfFrames", EditConditionHides, ClampMin = "0", ClampMax = "1"))
-	FPerPlatformFloat RemovalProportion;
-
-	/** Remove frames until worst distance error reaches this number of cm*/
-	UPROPERTY(EditAnywhere, Category = "Frame Removal", meta = (EditCondition = "bAllowFrameRemoval && FrameRemovalThresholdType == ACLFrameRemovalThresholdType::DistanceError", EditConditionHides, ClampMin = "0", ClampMax = "100"))
-	FPerPlatformFloat RemovalDistanceError;
-// @third party code - Epic Games End
-
 	//////////////////////////////////////////////////////////////////////////
 	// UAnimBoneCompressionCodec implementation
 // @third party code - Epic Games Begin
diff --git a/Engine/Plugins/Animation/ACLPlugin/Source/ACLPlugin/Private/AnimBoneCompressionCodec_ACL.cpp b/Engine/Plugins/Animation/ACLPlugin/Source/ACLPlugin/Private/AnimBoneCompressionCodec_ACL.cpp
--- a/Engine/Plugins/Animation/ACLPlugin/Source/ACLPlugin/Private/AnimBoneCompressionCodec_ACL.cpp
+++ b/Engine/Plugins/Animation/ACLPlugin/Source/ACLPlugin/Private/AnimBoneCompressionCodec_ACL.cpp
@@ -27,13 +27,6 @@ THIRD_PARTY_INCLUDES_END
 UAnimBoneCompressionCodec_ACL::UAnimBoneCompressionCodec_ACL(const FObjectInitializer& ObjectInitializer)
 	: Super(ObjectInitializer)
 {
-#if WITH_EDITORONLY_DATA
-// @third party code - Epic Games Begin
-	bAllowFrameRemoval = false;
-	FrameRemovalThresholdType = ACLFrameRemovalThresholdType::ProportionOfFrames;
-	RemovalProportion = 0.5f;
-// @third party code - Epic Games End
-#endif	// WITH_EDITORONLY_DATA
 }
 
 #if WITH_EDITORONLY_DATA
@@ -44,27 +37,6 @@ void UAnimBoneCompressionCodec_ACL::GetCompressionSettings(acl::compression_sett
 	OutSettings = acl::get_default_compression_settings();
 
 	OutSettings.level = GetCompressionLevel(CompressionLevel);
-// @third party code - Epic Games Begin
-	if (bAllowFrameRemoval)
-	{
-		// set the compression settings for tracking error per frame
-		OutSettings.enable_frame_stripping = true;
-		OutSettings.frame_stripping_use_proportion = (FrameRemovalThresholdType == ACLFrameRemovalThresholdType::ProportionOfFrames);
-		if (TargetPlatform != nullptr)
-		{
-			FName TargetPlatformName = TargetPlatform->GetTargetPlatformInfo().Name;
-			OutSettings.frame_stripping_proportion = RemovalProportion.GetValueForPlatform(TargetPlatformName);
-			OutSettings.frame_stripping_error_distance = RemovalDistanceError.GetValueForPlatform(TargetPlatformName);
-		}
-		else
-		{
-			OutSettings.frame_stripping_proportion = RemovalProportion.GetValue();
-			OutSettings.frame_stripping_error_distance = RemovalDistanceError.GetValue();
-		}
-
-		OutSettings.metadata.include_contributing_error = true;
-	}
-// @third party code - Epic Games End
 }
 
 // @third party code - Epic Games Begin
