Disable ACL_ASSERT until we fix the incompatible string parameters (%s is passing char strings in ACL code)

#rnx

diff --git a/Engine/Plugins/Animation/ACLPlugin/Source/ACLPlugin/Public/ACLImpl.h b/Engine/Plugins/Animation/ACLPlugin/Source/ACLPlugin/Public/ACLImpl.h
--- a/Engine/Plugins/Animation/ACLPlugin/Source/ACLPlugin/Public/ACLImpl.h
+++ b/Engine/Plugins/Animation/ACLPlugin/Source/ACLPlugin/Public/ACLImpl.h
@@ -13,8 +13,12 @@
 	// developers use a Development configuration for the editor. The ACL library runs extensive
 	// unit and regression tests on a very large number of clips which minimizes the risk of
 	// having a legitimate assert fire.
-	#define ACL_ON_ASSERT_CUSTOM
-	#define ACL_ASSERT(expression, format, ...) checkf(expression, TEXT(format), #__VA_ARGS__)
+	// 
+	// ACL assert strings are currently incompatible with checkf since it uses %s but passes a char-type string. In UE %s requires to pass a TCHAR string.
+	//#define ACL_ON_ASSERT_CUSTOM
+
+	
+	//#define ACL_ASSERT(expression, format, ...) checkf(expression, TEXT(format), ##__VA_ARGS__)
 #endif
 
 //////////////////////////////////////////////////////////////////////////
