Disable keyframe stripping when ACL database codec is used

diff --git a/Engine/Plugins/Animation/ACLPlugin/Source/ACLPlugin/Private/AnimBoneCompressionCodec_ACLDatabase.cpp b/Engine/Plugins/Animation/ACLPlugin/Source/ACLPlugin/Private/AnimBoneCompressionCodec_ACLDatabase.cpp
--- a/Engine/Plugins/Animation/ACLPlugin/Source/ACLPlugin/Private/AnimBoneCompressionCodec_ACLDatabase.cpp
+++ b/Engine/Plugins/Animation/ACLPlugin/Source/ACLPlugin/Private/AnimBoneCompressionCodec_ACLDatabase.cpp
@@ -176,6 +176,9 @@ void UAnimBoneCompressionCodec_ACLDatabase::GetCompressionSettings(const class I
 
 	OutSettings.level = GetCompressionLevel(CompressionLevel);
 	OutSettings.enable_database_support = true;
+
+	// Disable keyframe stripping, even the trivial one as it currently isn't supported
+	OutSettings.keyframe_stripping.strip_trivial = false;
 }
 
 // @third party code - Epic Games Begin
