Optimize AnimNext in various ways
- Avoid initializing a pose we'll write to immediately afterwards
- Modify ACL to output SOA poses natively
- Remove dummy bone container
- Clean up update/evaluate traversals to reduce churn

Enable named events in Test config of AnimNextSample project

diff --git a/Engine/Plugins/Animation/ACLPlugin/Source/ACLPlugin/Classes/AnimBoneCompressionCodec_ACL.h b/Engine/Plugins/Animation/ACLPlugin/Source/ACLPlugin/Classes/AnimBoneCompressionCodec_ACL.h
--- a/Engine/Plugins/Animation/ACLPlugin/Source/ACLPlugin/Classes/AnimBoneCompressionCodec_ACL.h
+++ b/Engine/Plugins/Animation/ACLPlugin/Source/ACLPlugin/Classes/AnimBoneCompressionCodec_ACL.h
@@ -39,5 +39,6 @@ class UAnimBoneCompressionCodec_ACL : public UAnimBoneCompressionCodec_ACLBase
 
 	// UAnimBoneCompressionCodec implementation
 	virtual void DecompressPose(FAnimSequenceDecompressionContext& DecompContext, const BoneTrackArray& RotationPairs, const BoneTrackArray& TranslationPairs, const BoneTrackArray& ScalePairs, TArrayView<FTransform>& OutAtoms) const override;
+	virtual void DecompressPose(FAnimSequenceDecompressionContext& DecompContext, const UE::Anim::FAnimPoseDecompressionData& DecompressionData) const override;
 	virtual void DecompressBone(FAnimSequenceDecompressionContext& DecompContext, int32 TrackIndex, FTransform& OutAtom) const override;
 };
diff --git a/Engine/Plugins/Animation/ACLPlugin/Source/ACLPlugin/Private/ACLDecompressionImpl.h b/Engine/Plugins/Animation/ACLPlugin/Source/ACLPlugin/Private/ACLDecompressionImpl.h
--- a/Engine/Plugins/Animation/ACLPlugin/Source/ACLPlugin/Private/ACLDecompressionImpl.h
+++ b/Engine/Plugins/Animation/ACLPlugin/Source/ACLPlugin/Private/ACLDecompressionImpl.h
@@ -3,6 +3,7 @@
 
 // Copyright 2018 Nicholas Frechette. All Rights Reserved.
 
+#include "Animation/AnimBoneDecompressionData.h"
 #include "Animation/AnimSequenceDecompressionContext.h"
 #include "Animation/AnimTypes.h"
 #include "AnimEncoding.h"
@@ -156,6 +157,111 @@ struct FUE4OutputWriter final : public acl::track_writer
 	}
 };
 
+/*
+ * Output pose writer that can selectively skip certain tracks.
+ */
+template<bool bUseBindPose>
+struct FUEOutputWriterSoA final : public acl::track_writer
+{
+	// Raw pointer for performance reasons, caller is responsible for ensuring data is valid
+
+	// The bind pose
+	const FTransform* RefPoses;
+
+	// Maps track indices to bone indices
+	const FTrackToSkeletonMap* TrackToBoneMapping;
+
+	// The track to output transform index map
+	const FAtomIndices* TrackToAtomsMap;
+
+	// The output transforms we write
+	FQuat* Rotations;
+	FVector* Translations;
+	FVector* Scales3D;
+
+	FUEOutputWriterSoA(
+		const FAtomIndices* InTrackToAtomsMap,
+		TArrayView<FQuat>& InRotations,
+		TArrayView<FVector>& InTranslations,
+		TArrayView<FVector>& InScales3D)
+		: RefPoses(nullptr)
+		, TrackToBoneMapping(nullptr)
+		, TrackToAtomsMap(InTrackToAtomsMap)
+		, Rotations(InRotations.GetData())
+		, Translations(InTranslations.GetData())
+		, Scales3D(InScales3D.GetData())
+	{}
+
+	FUEOutputWriterSoA(
+		const TArrayView<const FTransform>& InRefPoses,
+		const TArrayView<const FTrackToSkeletonMap>& TrackToSkeletonMap,
+		const FAtomIndices* InTrackToAtomsMap,
+		TArrayView<FQuat>& InRotations,
+		TArrayView<FVector>& InTranslations,
+		TArrayView<FVector>& InScales3D)
+		: RefPoses(InRefPoses.GetData())
+		, TrackToBoneMapping(TrackToSkeletonMap.GetData())
+		, TrackToAtomsMap(InTrackToAtomsMap)
+		, Rotations(InRotations.GetData())
+		, Translations(InTranslations.GetData())
+		, Scales3D(InScales3D.GetData())
+	{}
+
+	//////////////////////////////////////////////////////////////////////////
+	// Override the OutputWriter behavior
+	// If we use the bind pose, each default sub-track will grab the bind pose value
+	// Otherwise we output the constant default values
+	FORCEINLINE_DEBUGGABLE bool skip_track_rotation(uint32_t BoneIndex) const { return TrackToAtomsMap[BoneIndex].Rotation == 0xFFFF; }
+	FORCEINLINE_DEBUGGABLE bool skip_track_translation(uint32_t BoneIndex) const { return TrackToAtomsMap[BoneIndex].Translation == 0xFFFF; }
+	FORCEINLINE_DEBUGGABLE bool skip_track_scale(uint32_t BoneIndex) const { return TrackToAtomsMap[BoneIndex].Scale == 0xFFFF; }
+
+	static constexpr acl::default_sub_track_mode get_default_rotation_mode() { return bUseBindPose ? acl::default_sub_track_mode::variable : acl::default_sub_track_mode::constant; }
+	static constexpr acl::default_sub_track_mode get_default_translation_mode() { return bUseBindPose ? acl::default_sub_track_mode::variable : acl::default_sub_track_mode::constant; }
+	// Always legacy for scale since there is no scale present in the bind pose
+	static constexpr acl::default_sub_track_mode get_default_scale_mode() { return acl::default_sub_track_mode::legacy; }
+
+	// TODO: There is a performance impact here because the ref pose might use doubles on PC
+	FORCEINLINE_DEBUGGABLE rtm::quatf RTM_SIMD_CALL get_variable_default_rotation(uint32_t TrackIndex) const
+	{
+		return UEQuatToACL(RefPoses[TrackToBoneMapping[TrackIndex].BoneTreeIndex].GetRotation());
+	}
+
+	FORCEINLINE_DEBUGGABLE rtm::vector4f RTM_SIMD_CALL get_variable_default_translation(uint32_t TrackIndex) const
+	{
+		return UEVector3ToACL(RefPoses[TrackToBoneMapping[TrackIndex].BoneTreeIndex].GetTranslation());
+	}
+
+	//////////////////////////////////////////////////////////////////////////
+	// Called by the decoder to write out a quaternion rotation value for a specified bone index
+	FORCEINLINE_DEBUGGABLE void RTM_SIMD_CALL write_rotation(uint32_t BoneIndex, rtm::quatf_arg0 Rotation)
+	{
+		const uint32 AtomIndex = TrackToAtomsMap[BoneIndex].Rotation;
+
+		FQuat& OutRotation = Rotations[AtomIndex];
+		OutRotation = (FQuat)ACLQuatToUE(Rotation);
+	}
+
+	//////////////////////////////////////////////////////////////////////////
+	// Called by the decoder to write out a translation value for a specified bone index
+	FORCEINLINE_DEBUGGABLE void RTM_SIMD_CALL write_translation(uint32_t BoneIndex, rtm::vector4f_arg0 Translation)
+	{
+		const uint32 AtomIndex = TrackToAtomsMap[BoneIndex].Translation;
+
+		FVector& OutTranslation = Translations[AtomIndex];
+		OutTranslation = (FVector)ACLVector3ToUE(Translation);
+	}
+
+	//////////////////////////////////////////////////////////////////////////
+	// Called by the decoder to write out a scale value for a specified bone index
+	FORCEINLINE_DEBUGGABLE void RTM_SIMD_CALL write_scale(uint32_t BoneIndex, rtm::vector4f_arg0 Scale)
+	{
+		const uint32 AtomIndex = TrackToAtomsMap[BoneIndex].Scale;
+
+		FVector& OutScale = Scales3D[AtomIndex];
+		OutScale = (FVector)ACLVector3ToUE(Scale);
+	}
+};
+
 /*
 * Output track writer for a single track.
 */
@@ -373,3 +479,76 @@ FORCEINLINE_DEBUGGABLE void DecompressPose(FAnimSequenceDecompressionContext& De
 		ACLContext.decompress_tracks(PoseWriter);
 	}
 }
+
+template<class ACLContextType>
+FORCEINLINE_DEBUGGABLE void DecompressPose(
+	const FAnimSequenceDecompressionContext& DecompContext,
+	ACLContextType& ACLContext,
+	const UE::Anim::FAnimPoseDecompressionData& DecompressionData)
+{
+	const float Time = DecompContext.GetEvaluationTime();
+
+	// Seek first, we'll start prefetching ahead right away
+	ACLContext.seek(Time, get_rounding_policy(DecompContext.Interpolation));
+
+	const acl::compressed_tracks* CompressedClipData = ACLContext.get_compressed_tracks();
+	const int32 TrackCount = CompressedClipData->get_num_tracks();
+
+	// TODO: Allocate this with padding and use SIMD to set everything to 0xFF
+	FAtomIndices* TrackToAtomsMap = new(FMemStack::Get()) FAtomIndices[TrackCount];
+	FMemory::Memset(TrackToAtomsMap, 0xFF, sizeof(FAtomIndices) * TrackCount);
+
+	// TODO: We should only need 1x uint16 atom index for each track/bone index
+	// and we need 3 bits to tell whether we care about the rot/trans/scale
+	// The rot and scale pairs are often the same array, skip the scale iteration and set both flags while
+	// iterating on the rotation pairs. This will reduce the mapping size by 2 bytes if we use 4 bytes.
+	// All reads will be aligned, can we pack further with 1x uint16 and 1x uint8 and do unaligned loads?
+	// Need to double check what the ASM looks like on x64 and ARM first to make sure it's good
+	// Maybe having two arrays side by side is better with uint16/uint8? Or having 1 bitset array?
+	// Ultimately, when we load these indices, they will be in the L1 since we write them here just before
+	// we use them during decompression. Optimizing for quick loading/unpacking it best.
+
+	for (const BoneTrackPair& Pair : DecompressionData.GetRotationPairs())
+	{
+		TrackToAtomsMap[Pair.TrackIndex].Rotation = (uint16)Pair.AtomIndex;
+	}
+
+	for (const BoneTrackPair& Pair : DecompressionData.GetTranslationPairs())
+	{
+		TrackToAtomsMap[Pair.TrackIndex].Translation = (uint16)Pair.AtomIndex;
+	}
+
+	const acl::acl_impl::tracks_header& TracksHeader = acl::acl_impl::get_tracks_header(*CompressedClipData);
+	if (TracksHeader.get_has_scale())
+	{
+		for (const BoneTrackPair& Pair : DecompressionData.GetScalePairs())
+		{
+			TrackToAtomsMap[Pair.TrackIndex].Scale = (uint16)Pair.AtomIndex;
+		}
+	}
+
+	// We will decompress the whole pose even if we only care about a smaller subset of bone tracks.
+	// This ensures we read the compressed pose data once, linearly.
+
+	// See [Bind pose stripping] for details
+	// Are we non-additive?
+	if (CompressedClipData->get_default_scale() != 0)
+	{
+		// Non-additive anim sequences must write out the bind pose for default sub-tracks since they
+		// have been stripped from the data. Additive anim sequences always have the additive identity
+		// stripped, no need for the bind pose.
+		constexpr bool bUseBindPose = true;
+
+		FUEOutputWriterSoA<bUseBindPose> PoseWriter(DecompContext.GetRefLocalPoses(), DecompContext.GetTrackToSkeletonMap(), TrackToAtomsMap, DecompressionData.GetOutAtomRotations(), DecompressionData.GetOutAtomTranslations(), DecompressionData.GetOutAtomScales3D());
+		ACLContext.decompress_tracks(PoseWriter);
+	}
+	else
+	{
+		// Additive anim sequences have the identity stripped out, we'll output it.
+		// We also output the regular identity if bind pose stripping isn't enabled.
+		constexpr bool bUseBindPose = false;
+
+		FUEOutputWriterSoA<bUseBindPose> PoseWriter(TrackToAtomsMap, DecompressionData.GetOutAtomRotations(), DecompressionData.GetOutAtomTranslations(), DecompressionData.GetOutAtomScales3D());
+		ACLContext.decompress_tracks(PoseWriter);
+	}
+}
diff --git a/Engine/Plugins/Animation/ACLPlugin/Source/ACLPlugin/Private/AnimBoneCompressionCodec_ACL.cpp b/Engine/Plugins/Animation/ACLPlugin/Source/ACLPlugin/Private/AnimBoneCompressionCodec_ACL.cpp
--- a/Engine/Plugins/Animation/ACLPlugin/Source/ACLPlugin/Private/AnimBoneCompressionCodec_ACL.cpp
+++ b/Engine/Plugins/Animation/ACLPlugin/Source/ACLPlugin/Private/AnimBoneCompressionCodec_ACL.cpp
@@ -75,6 +75,18 @@ void UAnimBoneCompressionCodec_ACL::DecompressPose(FAnimSequenceDecompressionCon
 	::DecompressPose(DecompContext, ACLContext, RotationPairs, TranslationPairs, ScalePairs, OutAtoms);
 }
 
+void UAnimBoneCompressionCodec_ACL::DecompressPose(FAnimSequenceDecompressionContext& DecompContext, const UE::Anim::FAnimPoseDecompressionData& DecompressionData) const
+{
+	const FACLCompressedAnimData& AnimData = static_cast<const FACLCompressedAnimData&>(DecompContext.CompressedAnimData);
+	const acl::compressed_tracks* CompressedClipData = AnimData.GetCompressedTracks();
+	check(CompressedClipData != nullptr && CompressedClipData->is_valid(false).empty());
+
+	acl::decompression_context<UE4DefaultDecompressionSettings> ACLContext;
+	ACLContext.initialize(*CompressedClipData);
+
+	::DecompressPose(DecompContext, ACLContext, DecompressionData);
+}
+
 void UAnimBoneCompressionCodec_ACL::DecompressBone(FAnimSequenceDecompressionContext& DecompContext, int32 TrackIndex, FTransform& OutAtom) const
 {
 	const FACLCompressedAnimData& AnimData = static_cast<const FACLCompressedAnimData&>(DecompContext.CompressedAnimData);
diff --git a/Engine/Plugins/Experimental/AnimNext/Source/AnimNext/Private/DecoratorInterfaces/IEvaluate.cpp b/Engine/Plugins/Experimental/AnimNext/Source/AnimNext/Private/DecoratorInterfaces/IEvaluate.cpp
--- a/Engine/Plugins/Experimental/AnimNext/Source/AnimNext/Private/DecoratorInterfaces/IEvaluate.cpp
+++ b/Engine/Plugins/Experimental/AnimNext/Source/AnimNext/Private/DecoratorInterfaces/IEvaluate.cpp
@@ -65,10 +65,35 @@ namespace UE::AnimNext
 		PostEvaluate,
 	};
 
+	// This structure is transient and lives either on the stack or the memstack and its destructor may not be called
 	struct FEvaluateEntry
 	{
+		// The decorator handle that points to our node to update
 		FWeakDecoratorPtr	DecoratorPtr;
+
+		// Which step we wish to perform when we next see this entry
 		EEvaluateStep		DesiredStep = EEvaluateStep::PreEvaluate;
+
+		// Once we've called PreUpdate, we cache the decorator binding to avoid a redundant query to call PostUpdate
+		TDecoratorBinding<IEvaluate> EvaluateDecorator;
+
+		// These pointers are mutually exclusive
+		// An entry is either part of the pending stack, the free list, or neither
+		union
+		{
+			// Next entry in the linked list of free entries
+			FEvaluateEntry* NextFreeEntry = nullptr;
+
+			// Previous entry below us in the nodes pending stack
+			FEvaluateEntry* PrevStackEntry;
+		};
+
+		FEvaluateEntry(const FWeakDecoratorPtr& InDecoratorPtr, EEvaluateStep InDesiredStep, FEvaluateEntry* InPrevStackEntry = nullptr)
+			: DecoratorPtr(InDecoratorPtr)
+			, DesiredStep(InDesiredStep)
+			, PrevStackEntry(InPrevStackEntry)
+		{
+		}
 	};
 
 	FEvaluationProgram EvaluateGraph(FExecutionContext& Context, FEvaluateTraversalContext& TraversalContext, FWeakDecoratorPtr GraphRootPtr)
@@ -82,48 +107,76 @@ namespace UE::AnimNext
 			return EvaluationProgram;	// Nothing to update
 		}
 
-		FMemMark Mark(FMemStack::Get());
-
-		TArray<FEvaluateEntry, TMemStackAllocator<>> NodesPendingUpdate;
-		NodesPendingUpdate.Reserve(64);
+		FMemStack& MemStack = FMemStack::Get();
+		FMemMark Mark(MemStack);
 
 		FChildrenArray Children;
-		Children.Reserve(64);
+		Children.Reserve(16);
 
 		FScopedEvaluationProgram ScopedEvaluationProgram(TraversalContext, EvaluationProgram);
 		FScopedTraversalContext ScopedTraversalContext(Context, TraversalContext);
 
 		// Add the graph root to kick start the evaluation process
-		NodesPendingUpdate.Push({ GraphRootPtr, EEvaluateStep::PreEvaluate });
+		FEvaluateEntry GraphRootEntry(GraphRootPtr, EEvaluateStep::PreEvaluate);
+		FEvaluateEntry* NodesPendingUpdateStackTop = &GraphRootEntry;
+
+		// List of free entries we can recycle
+		FEvaluateEntry* FreeEntryList = nullptr;
+
+		TDecoratorBinding<IHierarchy> HierarchyDecorator;
 
 		// Process every node twice: pre-evaluate and post-evaluate
-		while (!NodesPendingUpdate.IsEmpty())
+		while (NodesPendingUpdateStackTop != nullptr)
 		{
 			// Grab the top most entry
-			const bool bAllowShrinking = false;	// Don't allow shrinking to avoid churn
-			const FEvaluateEntry Entry = NodesPendingUpdate.Pop(bAllowShrinking);
+			FEvaluateEntry* Entry = NodesPendingUpdateStackTop;
+			bool bIsEntryUsed = true;
 
-			if (Entry.DesiredStep == EEvaluateStep::PreEvaluate)
+			if (Entry->DesiredStep == EEvaluateStep::PreEvaluate)
 			{
-				// This is the first time we visit this node, time to pre-evaluate
-				// Queue our node again so that post-evaluate is called afterwards
-				NodesPendingUpdate.Push({ Entry.DecoratorPtr, EEvaluateStep::PostEvaluate });
+				if (Context.GetInterface(Entry->DecoratorPtr, Entry->EvaluateDecorator))
+				{
+					// This is the first time we visit this node, time to pre-evaluate
+					Entry->EvaluateDecorator.PreEvaluate(Context);
 
-				TDecoratorBinding<IEvaluate> EvaluateDecorator;
-				if (Context.GetInterface(Entry.DecoratorPtr, EvaluateDecorator))
+					// Leave our entry on top of the stack, we'll need to call PostEvaluate once the children
+					// we'll push on top finish
+					Entry->DesiredStep = EEvaluateStep::PostEvaluate;
+				}
+				else
 				{
-					EvaluateDecorator.PreEvaluate(Context);
+					// This node doesn't implement IUpdate, we can pop it from the stack
+					NodesPendingUpdateStackTop = Entry->PrevStackEntry;
+					bIsEntryUsed = false;
 				}
 
-				TDecoratorBinding<IHierarchy> HierarchyDecorator;
-				if (Context.GetInterface(Entry.DecoratorPtr, HierarchyDecorator))
+				if (Context.GetInterface(Entry->DecoratorPtr, HierarchyDecorator))
 				{
 					HierarchyDecorator.GetChildren(Context, Children);
 
 					// Append our children in reserve order so that they are visited in the same order they were added
 					for (int32 ChildIndex = Children.Num() - 1; ChildIndex >= 0; --ChildIndex)
 					{
-						NodesPendingUpdate.Push({ Children[ChildIndex], EEvaluateStep::PreEvaluate });
+						// Insert our new child on top of the stack
+
+						FEvaluateEntry* ChildEntry;
+						if (FreeEntryList != nullptr)
+						{
+							// Grab an entry from the free list
+							ChildEntry = FreeEntryList;
+							FreeEntryList = ChildEntry->NextFreeEntry;
+
+							ChildEntry->DecoratorPtr = Children[ChildIndex];
+							ChildEntry->DesiredStep = EEvaluateStep::PreEvaluate;
+							ChildEntry->PrevStackEntry = NodesPendingUpdateStackTop;
+						}
+						else
+						{
+							// Allocate a new entry
+							ChildEntry = new(MemStack) FEvaluateEntry(Children[ChildIndex], EEvaluateStep::PreEvaluate, NodesPendingUpdateStackTop);
+						}
+
+						NodesPendingUpdateStackTop = ChildEntry;
 					}
 
 					// Reset our container for the next time we need it
@@ -135,16 +188,24 @@ namespace UE::AnimNext
 			}
 			else
 			{
-				// We've already visited this node once, time to post-evaluate
-				TDecoratorBinding<IEvaluate> EvaluateDecorator;
-				if (Context.GetInterface(Entry.DecoratorPtr, EvaluateDecorator))
-				{
-					EvaluateDecorator.PostEvaluate(Context);
-				}
+				// We've already visited this node once, time to post-update
+				check(Entry->EvaluateDecorator.IsValid());
+				Entry->EvaluateDecorator.PostEvaluate(Context);
+
+				// Now that we are done processing this entry, we can pop it
+				NodesPendingUpdateStackTop = Entry->PrevStackEntry;
+				bIsEntryUsed = false;
 
 				// Break and continue to the next top-most entry
 				// It is either a sibling read for its post-evaluate or our parent entry ready for its post-evaluate
 			}
+
+			if (!bIsEntryUsed)
+			{
+				// This entry is no longer used, add it to the free list
+				Entry->NextFreeEntry = FreeEntryList;
+				FreeEntryList = Entry;
+			}
 		}
 
 		return EvaluationProgram;
diff --git a/Engine/Plugins/Experimental/AnimNext/Source/AnimNext/Private/DecoratorInterfaces/IUpdate.cpp b/Engine/Plugins/Experimental/AnimNext/Source/AnimNext/Private/DecoratorInterfaces/IUpdate.cpp
--- a/Engine/Plugins/Experimental/AnimNext/Source/AnimNext/Private/DecoratorInterfaces/IUpdate.cpp
+++ b/Engine/Plugins/Experimental/AnimNext/Source/AnimNext/Private/DecoratorInterfaces/IUpdate.cpp
@@ -39,10 +39,35 @@ namespace UE::AnimNext
 		PostUpdate,
 	};
 
+	// This structure is transient and lives either on the stack or the memstack and its destructor may not be called
 	struct FUpdateEntry
 	{
+		// The decorator handle that points to our node to update
 		FWeakDecoratorPtr	DecoratorPtr;
+
+		// Which step we wish to perform when we next see this entry
 		EUpdateStep			DesiredStep = EUpdateStep::PreUpdate;
+
+		// Once we've called PreUpdate, we cache the decorator binding to avoid a redundant query to call PostUpdate
+		TDecoratorBinding<IUpdate> UpdateDecorator;
+
+		// These pointers are mutually exclusive
+		// An entry is either part of the pending stack, the free list, or neither
+		union
+		{
+			// Next entry in the linked list of free entries
+			FUpdateEntry* NextFreeEntry = nullptr;
+
+			// Previous entry below us in the nodes pending stack
+			FUpdateEntry* PrevStackEntry;
+		};
+
+		FUpdateEntry(const FWeakDecoratorPtr& InDecoratorPtr, EUpdateStep InDesiredStep, FUpdateEntry* InPrevStackEntry = nullptr)
+			: DecoratorPtr(InDecoratorPtr)
+			, DesiredStep(InDesiredStep)
+			, PrevStackEntry(InPrevStackEntry)
+		{
+		}
 	};
 
 	void UpdateGraph(FExecutionContext& Context, FUpdateTraversalContext& TraversalContext, FWeakDecoratorPtr GraphRootPtr)
@@ -54,31 +79,47 @@ namespace UE::AnimNext
 			return;	// Nothing to update
 		}
 
-		FMemMark Mark(FMemStack::Get());
-
-		TArray<FUpdateEntry, TMemStackAllocator<>> NodesPendingUpdate;
-		NodesPendingUpdate.Reserve(64);
+		FMemStack& MemStack = FMemStack::Get();
+		FMemMark Mark(MemStack);
 
 		FChildrenArray Children;
-		Children.Reserve(64);
+		Children.Reserve(16);
 
 		FScopedTraversalContext ScopedTraversalContext(Context, TraversalContext);
 
 		// Add the graph root to kick start the update process
-		NodesPendingUpdate.Push({ GraphRootPtr, EUpdateStep::PreUpdate });
+		FUpdateEntry GraphRootEntry(GraphRootPtr, EUpdateStep::PreUpdate);
+		FUpdateEntry* NodesPendingUpdateStackTop = &GraphRootEntry;
+
+		// List of free entries we can recycle
+		FUpdateEntry* FreeEntryList = nullptr;
+
+		TDecoratorBinding<IHierarchy> HierarchyDecorator;
 
 		// Process every node twice: pre-update and post-update
-		while (!NodesPendingUpdate.IsEmpty())
+		while (NodesPendingUpdateStackTop != nullptr)
 		{
 			// Grab the top most entry
-			const bool bAllowShrinking = false;	// Don't allow shrinking to avoid churn
-			const FUpdateEntry Entry = NodesPendingUpdate.Pop(bAllowShrinking);
+			FUpdateEntry* Entry = NodesPendingUpdateStackTop;
+			bool bIsEntryUsed = true;
 
-			if (Entry.DesiredStep == EUpdateStep::PreUpdate)
+			if (Entry->DesiredStep == EUpdateStep::PreUpdate)
 			{
-				// This is the first time we visit this node, time to pre-update
-				// Queue our node again so that post-update is called afterwards
-				NodesPendingUpdate.Push({ Entry.DecoratorPtr, EUpdateStep::PostUpdate });
+				if (Context.GetInterface(Entry->DecoratorPtr, Entry->UpdateDecorator))
+				{
+					// This is the first time we visit this node, time to pre-update
+					Entry->UpdateDecorator.PreUpdate(Context);
+
+					// Leave our entry on top of the stack, we'll need to call PostUpdate once the children
+					// we'll push on top finish
+					Entry->DesiredStep = EUpdateStep::PostUpdate;
+				}
+				else
+				{
+					// This node doesn't implement IUpdate, we can pop it from the stack
+					NodesPendingUpdateStackTop = Entry->PrevStackEntry;
+					bIsEntryUsed = false;
+				}
 
 				// Performance note
 				// When we process an animation graph for a frame, typically we'll update first before we evaluate
@@ -120,21 +161,33 @@ namespace UE::AnimNext
 				// perform as much useful work as possible while waiting for memory, hiding its slow latency by fully
 				// leveraging out-of-order CPU execution.
 
-				TDecoratorBinding<IUpdate> UpdateDecorator;
-				if (Context.GetInterface(Entry.DecoratorPtr, UpdateDecorator))
-				{
-					UpdateDecorator.PreUpdate(Context);
-				}
-
-				TDecoratorBinding<IHierarchy> HierarchyDecorator;
-				if (Context.GetInterface(Entry.DecoratorPtr, HierarchyDecorator))
+				if (Context.GetInterface(Entry->DecoratorPtr, HierarchyDecorator))
 				{
 					HierarchyDecorator.GetChildren(Context, Children);
 
 					// Append our children in reserve order so that they are visited in the same order they were added
 					for (int32 ChildIndex = Children.Num() - 1; ChildIndex >= 0; --ChildIndex)
 					{
-						NodesPendingUpdate.Push({ Children[ChildIndex], EUpdateStep::PreUpdate });
+						// Insert our new child on top of the stack
+
+						FUpdateEntry* ChildEntry;
+						if (FreeEntryList != nullptr)
+						{
+							// Grab an entry from the free list
+							ChildEntry = FreeEntryList;
+							FreeEntryList = ChildEntry->NextFreeEntry;
+
+							ChildEntry->DecoratorPtr = Children[ChildIndex];
+							ChildEntry->DesiredStep = EUpdateStep::PreUpdate;
+							ChildEntry->PrevStackEntry = NodesPendingUpdateStackTop;
+						}
+						else
+						{
+							// Allocate a new entry
+							ChildEntry = new(MemStack) FUpdateEntry(Children[ChildIndex], EUpdateStep::PreUpdate, NodesPendingUpdateStackTop);
+						}
+
+						NodesPendingUpdateStackTop = ChildEntry;
 					}
 
 					// Reset our container for the next time we need it
@@ -147,15 +200,23 @@ namespace UE::AnimNext
 			else
 			{
 				// We've already visited this node once, time to post-update
-				TDecoratorBinding<IUpdate> UpdateDecorator;
-				if (Context.GetInterface(Entry.DecoratorPtr, UpdateDecorator))
-				{
-					UpdateDecorator.PostUpdate(Context);
-				}
+				check(Entry->UpdateDecorator.IsValid());
+				Entry->UpdateDecorator.PostUpdate(Context);
+
+				// Now that we are done processing this entry, we can pop it
+				NodesPendingUpdateStackTop = Entry->PrevStackEntry;
+				bIsEntryUsed = false;
 
 				// Break and continue to the next top-most entry
 				// It is either a sibling read for its post-update or our parent entry ready for its post-update
 			}
+
+			if (!bIsEntryUsed)
+			{
+				// This entry is no longer used, add it to the free list
+				Entry->NextFreeEntry = FreeEntryList;
+				FreeEntryList = Entry;
+			}
 		}
 	}
 }
diff --git a/Engine/Plugins/Experimental/AnimNext/Source/AnimNext/Private/EvaluationVM/EvaluationVM.cpp b/Engine/Plugins/Experimental/AnimNext/Source/AnimNext/Private/EvaluationVM/EvaluationVM.cpp
--- a/Engine/Plugins/Experimental/AnimNext/Source/AnimNext/Private/EvaluationVM/EvaluationVM.cpp
+++ b/Engine/Plugins/Experimental/AnimNext/Source/AnimNext/Private/EvaluationVM/EvaluationVM.cpp
@@ -33,10 +33,12 @@ namespace UE::AnimNext
 		, CurrentLOD(InCurrentLOD)
 		, EvaluationFlags(InEvaluationFlags)
 	{
-		const UE::Anim::FCurveFilterSettings CurveFilterSettings(UE::Anim::ECurveFilterMode::DisallowAll);
-		if(USkeleton* Skeleton = const_cast<USkeleton*>(InReferencePose.Skeleton.Get())) // const_cast because the bone container takes a mutable reference
+		CurveFilter.SetFilterMode(UE::Anim::ECurveFilterMode::DisallowAll);
+
+		//const UE::Anim::FCurveFilterSettings CurveFilterSettings(UE::Anim::ECurveFilterMode::DisallowAll);
+		//if(USkeleton* Skeleton = const_cast<USkeleton*>(InReferencePose.Skeleton.Get())) // const_cast because the bone container takes a mutable reference
 		{
-			BoneContainer.InitializeTo(InReferencePose.GetLODBoneIndexToMeshBoneIndexMap(InCurrentLOD), CurveFilterSettings, *Skeleton);
+			//BoneContainer.InitializeTo(InReferencePose.GetLODBoneIndexToMeshBoneIndexMap(InCurrentLOD), CurveFilterSettings, *Skeleton);
 
 			// TODO: In AnimInstanceProxy this is how we initialize the bone container, we need to get the component somehow or we
 			// gotta figure out how to support ref pose overrides
@@ -59,7 +61,7 @@ namespace UE::AnimNext
 
 	bool FEvaluationVM::IsValid() const
 	{
-		return ReferencePose != nullptr && BoneContainer.IsValid();
+		return ReferencePose != nullptr;
 	}
 
 	void FEvaluationVM::Shrink()
@@ -86,7 +88,7 @@ namespace UE::AnimNext
 
 		if (EnumHasAnyFlags(EvaluationFlags, EEvaluationFlags::Curves))
 		{
-			Keyframe.Curves.InitFrom(BoneContainer);
+			Keyframe.Curves.SetFilter(&CurveFilter);
 		}
 
 		return Keyframe;
@@ -105,17 +107,12 @@ namespace UE::AnimNext
 
 		if (EnumHasAnyFlags(EvaluationFlags, EEvaluationFlags::Curves))
 		{
-			Keyframe.Curves.InitFrom(BoneContainer);
+			Keyframe.Curves.SetFilter(&CurveFilter);
 		}
 
 		return Keyframe;
 	}
 
-	const FBoneContainer& FEvaluationVM::GetBoneContainer() const
-	{
-		return BoneContainer;
-	}
-
 	FEvaluationVMStack& FEvaluationVM::GetOrCreateStack(const FEvaluationVMStackName& StackName, uint32 TypeID)
 	{
 		if (FEvaluationVMStack* Stack = InternalStacks.FindByHash(StackName.NameHash, StackName.Name))
diff --git a/Engine/Plugins/Experimental/AnimNext/Source/AnimNext/Private/EvaluationVM/Tasks/PushAnimSequenceKeyframe.cpp b/Engine/Plugins/Experimental/AnimNext/Source/AnimNext/Private/EvaluationVM/Tasks/PushAnimSequenceKeyframe.cpp
--- a/Engine/Plugins/Experimental/AnimNext/Source/AnimNext/Private/EvaluationVM/Tasks/PushAnimSequenceKeyframe.cpp
+++ b/Engine/Plugins/Experimental/AnimNext/Source/AnimNext/Private/EvaluationVM/Tasks/PushAnimSequenceKeyframe.cpp
@@ -56,12 +56,43 @@ void FAnimNextAnimSequenceKeyframeTask::Execute(UE::AnimNext::FEvaluationVM& VM)
 
 		if (EnumHasAnyFlags(VM.GetFlags(), EEvaluationFlags::Attributes))
 		{
-			FCompactPose Pose;		// Dummy but we need the bone container
-			Pose.SetBoneContainer(&VM.GetBoneContainer());
-
-			FAnimationPoseData PoseData(Pose, Keyframe.Curves, Keyframe.Attributes);
-
-			AnimSequencePtr->EvaluateAttributes(PoseData, ExtractionContext, bUseRawData);
+			// TODO: A few notes here:
+			// The attributes store a compact/lod bone index, this needs fixup if we feed the output into an AnimBP or whoever reads this later
+			// To be able to sample attributes, we need to sample bones as well, otherwise we'll have no refpose object
+			// We could assign the ref pose even if we don't sample bones, that would allow us to be more selective
+
+			QUICK_SCOPE_CYCLE_COUNTER(STAT_EvaluateAttributes);
+
+#if WITH_EDITOR
+			if (bUseRawData)
+			{
+				AnimSequencePtr->ValidateModel();
+
+				for (const FAnimatedBoneAttribute& Attribute : AnimSequencePtr->GetDataModel()->GetAttributes())
+				{
+					const int32 LODBoneIndex = Keyframe.Pose.GetRefPose().GetLODBoneIndexFromSkeletonBoneIndex(Attribute.Identifier.GetBoneIndex());
+					// Only add attribute if the bone its tied to exists in the currently evaluated set of bones
+					if (LODBoneIndex != INDEX_NONE)
+					{
+						UE::Anim::Attributes::GetAttributeValue(Keyframe.Attributes, FCompactPoseBoneIndex(LODBoneIndex), Attribute, ExtractionContext.CurrentTime);
+					}
+				}
+			}
+			else
+#endif // WITH_EDITOR
+			{
+				for (const TPair<FAnimationAttributeIdentifier, FAttributeCurve>& BakedAttribute : AnimSequencePtr->AttributeCurves)
+				{
+					const int32 LODBoneIndex = Keyframe.Pose.GetRefPose().GetLODBoneIndexFromSkeletonBoneIndex(BakedAttribute.Key.GetBoneIndex());
+					// Only add attribute if the bone its tied to exists in the currently evaluated set of bones
+					if (LODBoneIndex != INDEX_NONE)
+					{
+						UE::Anim::FAttributeId Info(BakedAttribute.Key.GetName(), FCompactPoseBoneIndex(LODBoneIndex));
+						uint8* AttributePtr = Keyframe.Attributes.FindOrAdd(BakedAttribute.Key.GetType(), Info);
+						BakedAttribute.Value.EvaluateToPtr(BakedAttribute.Key.GetType(), ExtractionContext.CurrentTime, AttributePtr);
+					}
+				}
+			}
 		}
 
 		VM.PushValue(KEYFRAME_STACK_NAME, MakeUnique<FKeyframeState>(MoveTemp(Keyframe)));
diff --git a/Engine/Plugins/Experimental/AnimNext/Source/AnimNext/Private/Scheduler/AnimNextScheduleGraphTask.cpp b/Engine/Plugins/Experimental/AnimNext/Source/AnimNext/Private/Scheduler/AnimNextScheduleGraphTask.cpp
--- a/Engine/Plugins/Experimental/AnimNext/Source/AnimNext/Private/Scheduler/AnimNextScheduleGraphTask.cpp
+++ b/Engine/Plugins/Experimental/AnimNext/Source/AnimNext/Private/Scheduler/AnimNextScheduleGraphTask.cpp
@@ -88,7 +88,12 @@ void FAnimNextScheduleGraphTask::RunGraph(const UE::AnimNext::FScheduleContext&
 	FParamStack::FPushedLayerHandle LayerHandle = ParamStack.PushLayer(InstanceData.GraphInputLayers[TaskIndex]);
 
 	// Internally we use memstack allocation, so we need a mark here
-	FMemMark MemMark(FMemStack::Get());
+	FMemStack& MemStack = FMemStack::Get();
+	FMemMark MemMark(MemStack);
+
+	// We allocate a dummy buffer to trigger the allocation of a large chunk if this is the first mark
+	// This reduces churn internally by avoiding a chunk to be repeatedly allocated and freed as we push/pop marks
+	MemStack.Alloc(size_t(FPageAllocator::SmallPageSize) + 1, 16);
 
 	const FContext Context(DeltaTime);
 	Graph->Run(Context, InstanceData.GraphInstanceData[TaskIndex], EAnimNextGraphSimulationSteps::All);
diff --git a/Engine/Plugins/Experimental/AnimNext/Source/AnimNext/Public/EvaluationVM/EvaluationVM.h b/Engine/Plugins/Experimental/AnimNext/Source/AnimNext/Public/EvaluationVM/EvaluationVM.h
--- a/Engine/Plugins/Experimental/AnimNext/Source/AnimNext/Public/EvaluationVM/EvaluationVM.h
+++ b/Engine/Plugins/Experimental/AnimNext/Source/AnimNext/Public/EvaluationVM/EvaluationVM.h
@@ -4,7 +4,7 @@
 
 #include "CoreMinimal.h"
 
-#include "BoneContainer.h"
+#include "Animation/AnimCurveFilter.h"
 #include "EvaluationVM/EvaluationFlags.h"
 #include "EvaluationVM/KeyframeState.h"
 #include "ReferencePose.h"
@@ -185,9 +185,6 @@ namespace UE::AnimNext
 		// Returns an uninitialized keyframe with memory pre-allocated
 		[[nodiscard]] FKeyframeState MakeUninitializedKeyframe(bool bAdditiveKeyframe) const;
 
-		// Returns the bone container we are evaluating with
-		[[nodiscard]] const FBoneContainer& GetBoneContainer() const;
-
 	private:
 		// Disallow copy
 		FEvaluationVM(const FEvaluationVM&) = delete;
@@ -205,9 +202,6 @@ namespace UE::AnimNext
 		// Various internal stacks that tasks can use (e.g. keyframe state)
 		TMap<FName, FEvaluationVMStack> InternalStacks;
 
-		// The bone container we are evaluating with
-		FBoneContainer BoneContainer;
-
 		// Reference pose
 		const UE::AnimNext::FReferencePose* ReferencePose = nullptr;
 
@@ -217,6 +211,9 @@ namespace UE::AnimNext
 		// Flags that control what we wish to evaluate
 		EEvaluationFlags EvaluationFlags = EEvaluationFlags::All;
 
+		// Default curve filter
+		UE::Anim::FCurveFilter CurveFilter;
+
 		// TODO: Use a stack based allocator for the stack entries and mark freed entries on Pop and coalesce when we can
 		// Can allocate stack segments from memstack, limit allocation size to 1024 bytes, allocate segments of 8192 bytes
 	};
diff --git a/Engine/Plugins/Experimental/AnimNext/Source/AnimNext/Public/LODPose.h b/Engine/Plugins/Experimental/AnimNext/Source/AnimNext/Public/LODPose.h
--- a/Engine/Plugins/Experimental/AnimNext/Source/AnimNext/Public/LODPose.h
+++ b/Engine/Plugins/Experimental/AnimNext/Source/AnimNext/Public/LODPose.h
@@ -195,7 +195,7 @@ struct TLODPose : public FLODPose
 		RefPose = &InRefPose;
 
 		const int32 NumTransforms = InRefPose.GetNumBonesForLOD(InLODLevel);
-		LocalTransforms.SetNum(NumTransforms);
+		LocalTransforms.SetNumUninitialized(NumTransforms);
 		LocalTransformsView = LocalTransforms.GetView();
 
 		Flags = (ELODPoseFlags)(bAdditive ? (Flags | ELODPoseFlags::Additive) : Flags & ELODPoseFlags::Additive);
diff --git a/Engine/Plugins/Experimental/AnimNext/Source/AnimNext/Public/TransformArray.h b/Engine/Plugins/Experimental/AnimNext/Source/AnimNext/Public/TransformArray.h
--- a/Engine/Plugins/Experimental/AnimNext/Source/AnimNext/Public/TransformArray.h
+++ b/Engine/Plugins/Experimental/AnimNext/Source/AnimNext/Public/TransformArray.h
@@ -41,10 +41,14 @@ struct TTransformArrayAoS
 
 	void SetNum(int32 NumTransforms, bool bAllowShrinking = true)
 	{
-		constexpr int32 TransformSize = sizeof(FVector) + sizeof(FQuat) + sizeof(FVector);
 		Transforms.SetNum(NumTransforms, bAllowShrinking);
 	}
 
+	void SetNumUninitialized(int32 NumTransforms, bool bAllowShrinking = true)
+	{
+		Transforms.SetNumUninitialized(NumTransforms, bAllowShrinking);
+	}
+
 	inline void SetIdentity(bool bAdditiveIdentity = false)
 	{
 		std::fill(Transforms.begin(), Transforms.end(), bAdditiveIdentity ? TransformAdditiveIdentity : FTransform::Identity);
@@ -318,6 +322,14 @@ struct TTransformArraySoA
 		UpdateViews(AllocatedMemory.GetData(), NumTransforms);
 	}
 
+	void SetNumUninitialized(int32 NumTransforms, bool bAllowShrinking = true)
+	{
+		constexpr int32 TransformSize = sizeof(FVector) + sizeof(FQuat) + sizeof(FVector);
+		AllocatedMemory.SetNumUninitialized(NumTransforms * TransformSize, bAllowShrinking);
+
+		UpdateViews(AllocatedMemory.GetData(), NumTransforms);
+	}
+
 	inline void SetIdentity(bool bAdditiveIdentity = false)
 	{
 		std::fill(Translations.begin(), Translations.end(), FVector::ZeroVector);
diff --git a/Engine/Source/Runtime/Engine/Classes/Animation/AnimSequence.h b/Engine/Source/Runtime/Engine/Classes/Animation/AnimSequence.h
--- a/Engine/Source/Runtime/Engine/Classes/Animation/AnimSequence.h
+++ b/Engine/Source/Runtime/Engine/Classes/Animation/AnimSequence.h
@@ -921,4 +921,5 @@ public:
 	friend struct UE::Anim::Compression::FScopedCompressionGuard;
 	friend class FAnimDataControllerTestBase;
 	friend class UE::Anim::FAnimSequenceCompilingManager;
+	friend struct FAnimNextAnimSequenceKeyframeTask;
 };
diff --git a/Engine/Source/Runtime/Engine/Classes/Animation/AttributeCurve.h b/Engine/Source/Runtime/Engine/Classes/Animation/AttributeCurve.h
--- a/Engine/Source/Runtime/Engine/Classes/Animation/AttributeCurve.h
+++ b/Engine/Source/Runtime/Engine/Classes/Animation/AttributeCurve.h
@@ -190,6 +190,7 @@ protected:
 
 	friend class UAnimSequence;
 	friend struct UE::Anim::Attributes;
+	friend struct FAnimNextAnimSequenceKeyframeTask;
 };
 
 template<>
