sjson-cpp system include path

#rnx

diff --git a/Engine/Plugins/Animation/ACLPlugin/Source/ACLPlugin/ACLPlugin.Build.cs b/Engine/Plugins/Animation/ACLPlugin/Source/ACLPlugin/ACLPlugin.Build.cs
--- a/Engine/Plugins/Animation/ACLPlugin/Source/ACLPlugin/ACLPlugin.Build.cs
+++ b/Engine/Plugins/Animation/ACLPlugin/Source/ACLPlugin/ACLPlugin.Build.cs
@@ -18,6 +18,7 @@ namespace UnrealBuildTool.Rules
 
 			PublicSystemIncludePaths.Add(Path.Combine(ACLSDKDir, "acl/includes"));
 			PublicSystemIncludePaths.Add(Path.Combine(ACLSDKDir, "acl/external/rtm/includes"));
+			PublicSystemIncludePaths.Add(Path.Combine(ACLSDKDir, "acl/external/sjson-cpp/includes"));
 
 			PublicDependencyModuleNames.Add("Core");
 			PublicDependencyModuleNames.Add("CoreUObject");
