Expose bind pose scale to compression codecs
Fix ACLPlugin to use the bind pose scale

diff --git a/Engine/Plugins/Animation/ACLPlugin/Source/ACLPlugin/Private/ACLDecompressionImpl.h b/Engine/Plugins/Animation/ACLPlugin/Source/ACLPlugin/Private/ACLDecompressionImpl.h
--- a/Engine/Plugins/Animation/ACLPlugin/Source/ACLPlugin/Private/ACLDecompressionImpl.h
+++ b/Engine/Plugins/Animation/ACLPlugin/Source/ACLPlugin/Private/ACLDecompressionImpl.h
@@ -112,8 +112,7 @@ struct FUE4OutputWriter final : public acl::track_writer
 
 	static constexpr acl::default_sub_track_mode get_default_rotation_mode() { return bUseBindPose ? acl::default_sub_track_mode::variable : acl::default_sub_track_mode::constant; }
 	static constexpr acl::default_sub_track_mode get_default_translation_mode() { return bUseBindPose ? acl::default_sub_track_mode::variable : acl::default_sub_track_mode::constant; }
-	// Always legacy for scale since there is no scale present in the bind pose
-	static constexpr acl::default_sub_track_mode get_default_scale_mode() { return acl::default_sub_track_mode::legacy; }
+	static constexpr acl::default_sub_track_mode get_default_scale_mode() { return bUseBindPose ? acl::default_sub_track_mode::variable : acl::default_sub_track_mode::legacy; }
 
 	// TODO: There is a performance impact here because the ref pose might use doubles on PC
 	FORCEINLINE_DEBUGGABLE rtm::quatf RTM_SIMD_CALL get_variable_default_rotation(uint32_t TrackIndex) const
@@ -126,6 +125,11 @@ struct FUE4OutputWriter final : public acl::track_writer
 		return UEVector3ToACL(RefPoses[TrackToBoneMapping[TrackIndex].BoneTreeIndex].GetTranslation());
 	}
 
+	FORCEINLINE_DEBUGGABLE rtm::vector4f RTM_SIMD_CALL get_variable_default_scale(uint32_t TrackIndex) const
+	{
+		return UEVector3ToACL(RefPoses[TrackToBoneMapping[TrackIndex].BoneTreeIndex].GetScale3D());
+	}
+
 	//////////////////////////////////////////////////////////////////////////
 	// Called by the decoder to write out a quaternion rotation value for a specified bone index
 	FORCEINLINE_DEBUGGABLE void RTM_SIMD_CALL write_rotation(uint32_t BoneIndex, rtm::quatf_arg0 Rotation)
@@ -217,8 +221,7 @@ struct FUEOutputWriterSoA final : public acl::track_writer
 
 	static constexpr acl::default_sub_track_mode get_default_rotation_mode() { return bUseBindPose ? acl::default_sub_track_mode::variable : acl::default_sub_track_mode::constant; }
 	static constexpr acl::default_sub_track_mode get_default_translation_mode() { return bUseBindPose ? acl::default_sub_track_mode::variable : acl::default_sub_track_mode::constant; }
-	// Always legacy for scale since there is no scale present in the bind pose
-	static constexpr acl::default_sub_track_mode get_default_scale_mode() { return acl::default_sub_track_mode::legacy; }
+	static constexpr acl::default_sub_track_mode get_default_scale_mode() { return bUseBindPose ? acl::default_sub_track_mode::variable : acl::default_sub_track_mode::legacy; }
 
 	// TODO: There is a performance impact here because the ref pose might use doubles on PC
 	FORCEINLINE_DEBUGGABLE rtm::quatf RTM_SIMD_CALL get_variable_default_rotation(uint32_t TrackIndex) const
@@ -231,6 +234,11 @@ struct FUEOutputWriterSoA final : public acl::track_writer
 		return UEVector3ToACL(RefPoses[TrackToBoneMapping[TrackIndex].BoneTreeIndex].GetTranslation());
 	}
 
+	FORCEINLINE_DEBUGGABLE rtm::vector4f RTM_SIMD_CALL get_variable_default_scale(uint32_t TrackIndex) const
+	{
+		return UEVector3ToACL(RefPoses[TrackToBoneMapping[TrackIndex].BoneTreeIndex].GetScale3D());
+	}
+
 	//////////////////////////////////////////////////////////////////////////
 	// Called by the decoder to write out a quaternion rotation value for a specified bone index
 	FORCEINLINE_DEBUGGABLE void RTM_SIMD_CALL write_rotation(uint32_t BoneIndex, rtm::quatf_arg0 Rotation)
@@ -301,8 +309,7 @@ struct UE4OutputTrackWriter final : public acl::track_writer
 	// Otherwise we output the constant default values
 	static constexpr acl::default_sub_track_mode get_default_rotation_mode() { return bUseBindPose ? acl::default_sub_track_mode::variable : acl::default_sub_track_mode::constant; }
 	static constexpr acl::default_sub_track_mode get_default_translation_mode() { return bUseBindPose ? acl::default_sub_track_mode::variable : acl::default_sub_track_mode::constant; }
-	// Always legacy for scale since there is no scale present in the bind pose
-	static constexpr acl::default_sub_track_mode get_default_scale_mode() { return acl::default_sub_track_mode::legacy; }
+	static constexpr acl::default_sub_track_mode get_default_scale_mode() { return bUseBindPose ? acl::default_sub_track_mode::variable : acl::default_sub_track_mode::legacy; }
 
 	// TODO: There is a performance impact here because the ref pose might use doubles on PC
 	FORCEINLINE_DEBUGGABLE rtm::quatf RTM_SIMD_CALL get_variable_default_rotation(uint32_t TrackIndex) const
@@ -315,6 +322,11 @@ struct UE4OutputTrackWriter final : public acl::track_writer
 		return UEVector3ToACL(RefPoses[TrackToBoneMapping[TrackIndex].BoneTreeIndex].GetTranslation());
 	}
 
+	FORCEINLINE_DEBUGGABLE rtm::vector4f RTM_SIMD_CALL get_variable_default_scale(uint32_t TrackIndex) const
+	{
+		return UEVector3ToACL(RefPoses[TrackToBoneMapping[TrackIndex].BoneTreeIndex].GetScale3D());
+	}
+
 	//////////////////////////////////////////////////////////////////////////
 	// Called by the decoder to write out a quaternion rotation value for a specified bone index
 	FORCEINLINE_DEBUGGABLE void RTM_SIMD_CALL write_rotation(uint32_t BoneIndex, rtm::quatf_arg0 Rotation)
diff --git a/Engine/Plugins/Animation/ACLPlugin/Source/ACLPlugin/Private/ACLImpl.cpp b/Engine/Plugins/Animation/ACLPlugin/Source/ACLPlugin/Private/ACLImpl.cpp
--- a/Engine/Plugins/Animation/ACLPlugin/Source/ACLPlugin/Private/ACLImpl.cpp
+++ b/Engine/Plugins/Animation/ACLPlugin/Source/ACLPlugin/Private/ACLImpl.cpp
@@ -218,7 +218,7 @@ acl::track_array_qvvf BuildACLTransformTrackArray(ACLAllocator& AllocatorImpl, c
 			}
 			else
 			{
-				BindTransform = rtm::qvv_set(UEQuatToACL(UE4Bone.Orientation), UEVector3ToACL(UE4Bone.Position), ACLDefaultScale);
+				BindTransform = rtm::qvv_set(UEQuatToACL(UE4Bone.Orientation), UEVector3ToACL(UE4Bone.Position), UEVector3ToACL(UE4Bone.Scale));
 			}
 
 			for (uint32 SampleIndex = 0; SampleIndex < NumSamples; ++SampleIndex)
diff --git a/Engine/Plugins/Animation/ACLPlugin/Source/ACLPlugin/Private/AnimBoneCompressionCodec_ACLBase.cpp b/Engine/Plugins/Animation/ACLPlugin/Source/ACLPlugin/Private/AnimBoneCompressionCodec_ACLBase.cpp
--- a/Engine/Plugins/Animation/ACLPlugin/Source/ACLPlugin/Private/AnimBoneCompressionCodec_ACLBase.cpp
+++ b/Engine/Plugins/Animation/ACLPlugin/Source/ACLPlugin/Private/AnimBoneCompressionCodec_ACLBase.cpp
@@ -192,7 +192,6 @@ static void StripBindPose(const FCompressibleAnimData& CompressibleAnimData, acl
 	check(!CompressibleAnimData.bIsValidAdditive);
 
 	const int32 NumBones = CompressibleAnimData.BoneData.Num();
-	const rtm::vector4f DefaultScale = rtm::vector_set(1.0f);
 
 	for (int32 BoneIndex = 0; BoneIndex < NumBones; ++BoneIndex)
 	{
@@ -215,7 +214,7 @@ static void StripBindPose(const FCompressibleAnimData& CompressibleAnimData, acl
 		//         Single bone decompression will output the correct value
 
 		// Set the default value to the bind pose so that it can be stripped
-		Desc.default_value = rtm::qvv_set(UEQuatToACL(UE4Bone.Orientation), UEVector3ToACL(UE4Bone.Position), DefaultScale);
+		Desc.default_value = rtm::qvv_set(UEQuatToACL(UE4Bone.Orientation), UEVector3ToACL(UE4Bone.Position), UEVector3ToACL(UE4Bone.Scale));
 	}
 }
 
diff --git a/Engine/Source/Runtime/Engine/Private/Animation/AnimationUtils.cpp b/Engine/Source/Runtime/Engine/Private/Animation/AnimationUtils.cpp
--- a/Engine/Source/Runtime/Engine/Private/Animation/AnimationUtils.cpp
+++ b/Engine/Source/Runtime/Engine/Private/Animation/AnimationUtils.cpp
@@ -57,6 +57,7 @@ void FAnimationUtils::BuildSkeletonMetaData(USkeleton* Skeleton, TArray<FBoneDat
 
 		BoneData.Orientation = SrcTransform.GetRotation();
 		BoneData.Position = (FVector3f)SrcTransform.GetTranslation();
+		BoneData.Scale = (FVector3f)SrcTransform.GetScale3D();
 		BoneData.Name = RefSkeleton.GetBoneName(BoneIndex);
 
 		if ( BoneIndex > 0 )
diff --git a/Engine/Source/Runtime/Engine/Public/Animation/AnimTypes.h b/Engine/Source/Runtime/Engine/Public/Animation/AnimTypes.h
--- a/Engine/Source/Runtime/Engine/Public/Animation/AnimTypes.h
+++ b/Engine/Source/Runtime/Engine/Public/Animation/AnimTypes.h
@@ -1001,18 +1001,30 @@ public:
 class FBoneData
 {
 public:
+	/** The bind pose orientation. */
 	FQuat		Orientation;
-	FVector3f		Position;
+
+	/** The bind pose position. */
+	FVector3f	Position;
+
+	/** The bind pose scale. */
+	FVector3f	Scale;
+
 	/** Bone name. */
 	FName		Name;
+
 	/** Direct descendants.  Empty for end effectors. */
 	TArray<int32> Children;
+
 	/** List of bone indices from parent up to root. */
 	TArray<int32>	BonesToRoot;
+
 	/** List of end effectors for which this bone is an ancestor.  End effectors have only one element in this list, themselves. */
 	TArray<int32>	EndEffectors;
+
 	/** If a Socket is attached to that bone */
 	bool		bHasSocket;
+
 	/** If matched as a Key end effector */
 	bool		bKeyEndEffector;
 
@@ -1021,11 +1033,13 @@ public:
 	{
 		return GetDepth() ? BonesToRoot[0] : -1;
 	}
+
 	/**	@return		Distance to root; 0 for the root. */
 	int32 GetDepth() const
 	{
 		return BonesToRoot.Num();
 	}
+
 	/** @return		true if this bone is an end effector (has no children). */
 	bool IsEndEffector() const
 	{
