Add memory usage estimate to ACLPlugin codecs

diff --git a/Engine/Plugins/Animation/ACLPlugin/Source/ACLPlugin/Classes/AnimBoneCompressionCodec_ACLBase.h b/Engine/Plugins/Animation/ACLPlugin/Source/ACLPlugin/Classes/AnimBoneCompressionCodec_ACLBase.h
--- a/Engine/Plugins/Animation/ACLPlugin/Source/ACLPlugin/Classes/AnimBoneCompressionCodec_ACLBase.h
+++ b/Engine/Plugins/Animation/ACLPlugin/Source/ACLPlugin/Classes/AnimBoneCompressionCodec_ACLBase.h
@@ -73,6 +73,7 @@ class UAnimBoneCompressionCodec_ACLBase : public UAnimBoneCompressionCodec
 	// UAnimBoneCompressionCodec implementation
 	virtual bool Compress(const FCompressibleAnimData& CompressibleAnimData, FCompressibleAnimDataResult& OutResult) override;
 	virtual void PopulateDDCKey(const UE::Anim::Compression::FAnimDDCKeyArgs& KeyArgs, FArchive& Ar) override;
+	virtual int64 EstimateCompressionMemoryUsage(const UAnimSequence& AnimSequence) const override;
 
 	// Our implementation
 	virtual void PostCompression(const FCompressibleAnimData& CompressibleAnimData, FCompressibleAnimDataResult& OutResult) const {}
diff --git a/Engine/Plugins/Animation/ACLPlugin/Source/ACLPlugin/Private/AnimBoneCompressionCodec_ACLBase.cpp b/Engine/Plugins/Animation/ACLPlugin/Source/ACLPlugin/Private/AnimBoneCompressionCodec_ACLBase.cpp
--- a/Engine/Plugins/Animation/ACLPlugin/Source/ACLPlugin/Private/AnimBoneCompressionCodec_ACLBase.cpp
+++ b/Engine/Plugins/Animation/ACLPlugin/Source/ACLPlugin/Private/AnimBoneCompressionCodec_ACLBase.cpp
@@ -360,6 +360,36 @@ void UAnimBoneCompressionCodec_ACLBase::PopulateDDCKey(const UE::Anim::Compressi
 	}
 }
 
+int64 UAnimBoneCompressionCodec_ACLBase::EstimateCompressionMemoryUsage(const UAnimSequence& AnimSequence) const
+{
+	const int64 AnimSeqRawSize = AnimSequence.GetApproxRawSize();
+
+	int64 EstimatedMemoryUsage = 0;
+	EstimatedMemoryUsage += AnimSeqRawSize;	// We copy the raw data into the ACL format
+
+	if (AnimSequence.IsValidAdditive())
+	{
+		if (AnimSequence.RefPoseSeq)
+		{
+			// We copy the additive base into the ACL format
+			EstimatedMemoryUsage += AnimSequence.RefPoseSeq->GetApproxRawSize();
+		}
+		else
+		{
+			// We create the additive base in the ACL format, we use the same estimate as FAnimationSequenceAsyncCacheTask::GetRequiredMemoryEstimate()
+			EstimatedMemoryUsage += AnimSeqRawSize;
+		}
+	}
+
+	EstimatedMemoryUsage *= 2;	// Internally, ACL copies the raw data into a different format than the input because the input is not modified
+
+	EstimatedMemoryUsage += AnimSeqRawSize;	// ACL keeps a mutable copy of the lossy data that it modifies during compression
+	EstimatedMemoryUsage += AnimSeqRawSize;	// ACL will allocate the output buffer, assume that it's as large as the raw data
+	EstimatedMemoryUsage += 100 * 1024;		// Reserve 100 KB for internal bookkeeping and other required metadata
+
+	return EstimatedMemoryUsage;
+}
+
 ACLSafetyFallbackResult UAnimBoneCompressionCodec_ACLBase::ExecuteSafetyFallback(acl::iallocator& Allocator, const acl::compression_settings& Settings, const acl::track_array_qvvf& RawClip, const acl::track_array_qvvf& BaseClip, const acl::compressed_tracks& CompressedClipData, const FCompressibleAnimData& CompressibleAnimData, FCompressibleAnimDataResult& OutResult)
 {
 	return ACLSafetyFallbackResult::Ignored;
