Ensure the ACL plugin compression never fails

diff --git a/Engine/Plugins/Animation/ACLPlugin/Source/ACLPlugin/Classes/AnimBoneCompressionCodec_ACLBase.h b/Engine/Plugins/Animation/ACLPlugin/Source/ACLPlugin/Classes/AnimBoneCompressionCodec_ACLBase.h
--- a/Engine/Plugins/Animation/ACLPlugin/Source/ACLPlugin/Classes/AnimBoneCompressionCodec_ACLBase.h
+++ b/Engine/Plugins/Animation/ACLPlugin/Source/ACLPlugin/Classes/AnimBoneCompressionCodec_ACLBase.h
@@ -30,7 +30,16 @@ enum class ACLSafetyFallbackResult
 	Ignored,	// No safety fallback used
 };
 
-struct FACLCompressedAnimData final : public ICompressedAnimData
+struct FACLCompressedAnimDataBase : public ICompressedAnimData
+{
+	/** Whether or not compression failed */
+	bool bCompressionFailed = false;
+
+	// ICompressedAnimData implementation
+	virtual void SerializeCompressedData(UObject* DataOwner, FArchive& Ar) override;
+};
+
+struct FACLCompressedAnimData final : public FACLCompressedAnimDataBase
 {
 	/** Holds the compressed_tracks instance */
 	TArrayView<uint8> CompressedByteStream;
diff --git a/Engine/Plugins/Animation/ACLPlugin/Source/ACLPlugin/Classes/AnimBoneCompressionCodec_ACLDatabase.h b/Engine/Plugins/Animation/ACLPlugin/Source/ACLPlugin/Classes/AnimBoneCompressionCodec_ACLDatabase.h
--- a/Engine/Plugins/Animation/ACLPlugin/Source/ACLPlugin/Classes/AnimBoneCompressionCodec_ACLDatabase.h
+++ b/Engine/Plugins/Animation/ACLPlugin/Source/ACLPlugin/Classes/AnimBoneCompressionCodec_ACLDatabase.h
@@ -15,7 +15,7 @@ THIRD_PARTY_INCLUDES_END
 #include "AnimBoneCompressionCodec_ACLBase.h"
 #include "AnimBoneCompressionCodec_ACLDatabase.generated.h"
 
-struct FACLDatabaseCompressedAnimData final : public ICompressedAnimData
+struct FACLDatabaseCompressedAnimData final : public FACLCompressedAnimDataBase
 {
 	/** Maps the compressed_tracks instance. Used in cooked build only. */
 	TArrayView<uint8> CompressedByteStream;
@@ -45,7 +45,7 @@ struct FACLDatabaseCompressedAnimData final : public ICompressedAnimData
 	virtual ~FACLDatabaseCompressedAnimData() override;
 #endif
 
-	virtual void SerializeCompressedData(FArchive& Ar) override;
+	virtual void SerializeCompressedData(UObject* DataOwner, FArchive& Ar) override;
 	virtual void Bind(const TArrayView<uint8> BulkData) override;
 	virtual int64 GetApproxCompressedSize() const override;
 	virtual bool IsValid() const override;
diff --git a/Engine/Plugins/Animation/ACLPlugin/Source/ACLPlugin/Private/AnimBoneCompressionCodec_ACLBase.cpp b/Engine/Plugins/Animation/ACLPlugin/Source/ACLPlugin/Private/AnimBoneCompressionCodec_ACLBase.cpp
--- a/Engine/Plugins/Animation/ACLPlugin/Source/ACLPlugin/Private/AnimBoneCompressionCodec_ACLBase.cpp
+++ b/Engine/Plugins/Animation/ACLPlugin/Source/ACLPlugin/Private/AnimBoneCompressionCodec_ACLBase.cpp
@@ -27,6 +27,23 @@ THIRD_PARTY_INCLUDES_END
 
 #include <acl/core/compressed_tracks.h>
 
+void FACLCompressedAnimDataBase::SerializeCompressedData(UObject* DataOwner, FArchive& Ar)
+{
+	ICompressedAnimData::SerializeCompressedData(DataOwner, Ar);
+
+	Ar << bCompressionFailed;
+
+	if (bCompressionFailed && (Ar.IsLoading() || Ar.IsCooking()))
+	{
+		// Compression failed due to invalid settings or data (see below in Compress(..))
+		// We'll end up outputting the bind pose during decompression
+		// We report an error to cause the cook to fail
+		UE_LOG(LogAnimationCompression, Error,
+			TEXT("ACL failed to compress an anim sequence and will output the bind pose at runtime: %s"),
+			DataOwner != nullptr ? *DataOwner->GetPathName() : TEXT("[Unknown Sequence]"));
+	}
+}
+
 bool FACLCompressedAnimData::IsValid() const
 {
 	if (CompressedByteStream.Num() == 0)
@@ -218,13 +235,47 @@ static void StripBindPose(const FCompressibleAnimData& CompressibleAnimData, acl
 	}
 }
 
+static void ResetTracksToIdentity(const FCompressibleAnimData& CompressibleAnimData, bool bBuildAdditiveBase, acl::track_array_qvvf& ACLTracks)
+{
+	// This resets the input ACL tracks to the identity transform but retains all other values
+	const uint32 NumSamples = 1;
+	const float SampleRate = 30.0f;
+
+	// Additive animations have 0,0,0 scale as the default since we add it
+	const bool bIsAdditive = bBuildAdditiveBase ? false : CompressibleAnimData.bIsValidAdditive;
+	const FVector3f UE4DefaultScale(bIsAdditive ? 0.0f : 1.0f);
+	const rtm::vector4f ACLDefaultScale = rtm::vector_set(bIsAdditive ? 0.0f : 1.0f);
+
+	rtm::qvvf ACLIdentityTransform = rtm::qvv_identity();
+	ACLIdentityTransform.scale = ACLDefaultScale;
+
+	const acl::track_desc_transformf DefaultDesc;
+
+	for (acl::track_qvvf& ACLTrack : ACLTracks)
+	{
+		// Reset everything to the identity transform and default values
+		// Retain the output index to ensure proper output size
+		acl::track_desc_transformf Desc = ACLTrack.get_description();	// Copy
+		Desc.default_value = ACLIdentityTransform;
+		Desc.precision = DefaultDesc.precision;
+		Desc.shell_distance = DefaultDesc.shell_distance;
+		Desc.parent_index = acl::k_invalid_track_index;
+
+		// Reset track to a single sample
+		ACLTrack = acl::track_qvvf::make_reserve(Desc, ACLAllocatorImpl, NumSamples, SampleRate);
+		ACLTrack[0] = ACLIdentityTransform;
+	}
+}
+
 bool UAnimBoneCompressionCodec_ACLBase::Compress(const FCompressibleAnimData& CompressibleAnimData, FCompressibleAnimDataResult& OutResult)
 {
 	acl::track_array_qvvf ACLTracks = BuildACLTransformTrackArray(ACLAllocatorImpl, CompressibleAnimData, DefaultVirtualVertexDistance, SafeVirtualVertexDistance, false, PhantomTrackMode);
 
 	acl::track_array_qvvf ACLBaseTracks;
 	if (CompressibleAnimData.bIsValidAdditive)
+	{
 		ACLBaseTracks = BuildACLTransformTrackArray(ACLAllocatorImpl, CompressibleAnimData, DefaultVirtualVertexDistance, SafeVirtualVertexDistance, true, PhantomTrackMode);
+	}
 
 	UE_LOG(LogAnimationCompression, Verbose, TEXT("ACL Animation raw size: %u bytes [%s]"), ACLTracks.get_raw_size(), *CompressibleAnimData.FullName);
 
@@ -269,12 +320,43 @@ bool UAnimBoneCompressionCodec_ACLBase::Compress(const FCompressibleAnimData& Co
 
 	acl::output_stats Stats;
 	acl::compressed_tracks* CompressedTracks = nullptr;
-	const acl::error_result CompressionResult = acl::compress_track_list(ACLAllocatorImpl, ACLTracks, Settings, ACLBaseTracks, AdditiveFormat, CompressedTracks, Stats);
+	acl::error_result CompressionResult = acl::compress_track_list(ACLAllocatorImpl, ACLTracks, Settings, ACLBaseTracks, AdditiveFormat, CompressedTracks, Stats);
+
+	bool bEnableErrorReporting = true;
+	bool bCompressionFailed = false;
 
-	if (!CompressionResult.empty() || CompressedTracks == nullptr)
+	if (!CompressionResult.empty())
 	{
-		UE_LOG(LogAnimationCompression, Warning, TEXT("ACL failed to compress clip: %s [%s]"), ANSI_TO_TCHAR(CompressionResult.c_str()), *CompressibleAnimData.FullName);
-		return false;
+		// If compression failed, one of two things happened:
+		//    * Invalid settings were used, this would be a code/logic error that results in an improper usage of ACL
+		//    * Invalid data was provided, this would be a validation error that should ideally be caught earlier (e.g import, save)
+		// 
+		// Either way, if we get here, we cannot recover and we cannot fail as the engine assumes that compression always succeeds.
+		// We must handle failure gracefully. To that end, we compress an empty stub to ensure that something is present to
+		// decompress. Because the stub is empty, we'll simply output the bind pose. We still log this as an error to signal that
+		// this is a problem that needs to be fixed. This will allow the editor to continue working with the bind pose we'll output
+		// but cooking will fail preventing us from running with invalid state.
+
+		UE_LOG(LogAnimationCompression, Error, TEXT("ACL failed to compress anim sequence: %s [%s]"), ANSI_TO_TCHAR(CompressionResult.c_str()), *CompressibleAnimData.FullName);
+
+		// We reset the tracks to the identity, getting rid of any potentially invalid data.
+		// By setting them to the identity along with their default value as well, bind pose stripping will
+		// strip the single keyframe. This will result in the bind pose being outputted during decompression
+		// for non-additive animations and additive animations will retain the additive identity.
+		ResetTracksToIdentity(CompressibleAnimData, false, ACLTracks);
+		if (CompressibleAnimData.bIsValidAdditive)
+		{
+			ResetTracksToIdentity(CompressibleAnimData, true, ACLBaseTracks);
+		}
+
+		CompressionResult = acl::compress_track_list(ACLAllocatorImpl, ACLTracks, Settings, ACLBaseTracks, AdditiveFormat, CompressedTracks, Stats);
+
+		// The stub compression should never fail
+		check(CompressionResult.empty() && CompressedTracks != nullptr);
+
+		// Because we compress an empty stub, disable error reporting below
+		bEnableErrorReporting = false;
+		bCompressionFailed = true;
 	}
 
 	checkSlow(CompressedTracks->is_valid(true).empty());
@@ -291,7 +373,11 @@ bool UAnimBoneCompressionCodec_ACLBase::Compress(const FCompressibleAnimData& Co
 
 	OutResult.AnimData->CompressedNumberOfKeys = GetNumSamples(CompressibleAnimData);
 
+	FACLCompressedAnimDataBase& AnimData = static_cast<FACLCompressedAnimDataBase&>(*OutResult.AnimData);
+	AnimData.bCompressionFailed = bCompressionFailed;
+
 #if !NO_LOGGING
+	if (bEnableErrorReporting)
 	{
 		// Use debug settings in case codec picked is the fallback
 		acl::decompression_context<UE4DebugDecompressionSettings> Context;
@@ -321,7 +407,7 @@ void UAnimBoneCompressionCodec_ACLBase::PopulateDDCKey(const UE::Anim::Compressi
 {
 	Super::PopulateDDCKey(KeyArgs, Ar);
 
-	uint32 ForceRebuildVersion = 18;
+	uint32 ForceRebuildVersion = 19;
 
 	Ar << ForceRebuildVersion << DefaultVirtualVertexDistance << SafeVirtualVertexDistance << ErrorThreshold;
 	Ar << CompressionLevel;
diff --git a/Engine/Plugins/Animation/ACLPlugin/Source/ACLPlugin/Private/AnimBoneCompressionCodec_ACLDatabase.cpp b/Engine/Plugins/Animation/ACLPlugin/Source/ACLPlugin/Private/AnimBoneCompressionCodec_ACLDatabase.cpp
--- a/Engine/Plugins/Animation/ACLPlugin/Source/ACLPlugin/Private/AnimBoneCompressionCodec_ACLDatabase.cpp
+++ b/Engine/Plugins/Animation/ACLPlugin/Source/ACLPlugin/Private/AnimBoneCompressionCodec_ACLDatabase.cpp
@@ -26,9 +26,9 @@ THIRD_PARTY_INCLUDES_END
 
 #include UE_INLINE_GENERATED_CPP_BY_NAME(AnimBoneCompressionCodec_ACLDatabase)
 
-void FACLDatabaseCompressedAnimData::SerializeCompressedData(FArchive& Ar)
+void FACLDatabaseCompressedAnimData::SerializeCompressedData(UObject* DataOwner, FArchive& Ar)
 {
-	ICompressedAnimData::SerializeCompressedData(Ar);
+	FACLCompressedAnimDataBase::SerializeCompressedData(DataOwner, Ar);
 
 	Ar << SequenceNameHash;
 
diff --git a/Engine/Source/Runtime/Engine/Private/Animation/AnimCompressionTypes.cpp b/Engine/Source/Runtime/Engine/Private/Animation/AnimCompressionTypes.cpp
--- a/Engine/Source/Runtime/Engine/Private/Animation/AnimCompressionTypes.cpp
+++ b/Engine/Source/Runtime/Engine/Private/Animation/AnimCompressionTypes.cpp
@@ -1133,7 +1133,9 @@ FArchive& operator<<(FArchive& Ar, AnimationKeyFormat& Fmt)
 
 void FUECompressedAnimData::SerializeCompressedData(FArchive& Ar)
 {
+	PRAGMA_DISABLE_DEPRECATION_WARNINGS
 	ICompressedAnimData::SerializeCompressedData(Ar);
+	PRAGMA_ENABLE_DEPRECATION_WARNINGS
 
 	Ar << KeyEncodingFormat;
 	Ar << TranslationCompressionFormat;
@@ -1223,7 +1225,7 @@ void FUECompressedAnimDataMutable::BuildFinalBuffer(TArray<uint8>& OutCompressed
 	WriteArray(MemoryWriter, CompressedByteStream);
 }
 
-void ICompressedAnimData::SerializeCompressedData(class FArchive& Ar)
+void ICompressedAnimData::SerializeCompressedData(FArchive& Ar)
 {
 	Ar << CompressedNumberOfKeys;
 
@@ -1239,6 +1241,14 @@ void ICompressedAnimData::SerializeCompressedData(class FArchive& Ar)
 #endif
 }
 
+void ICompressedAnimData::SerializeCompressedData(UObject* DataOwner, FArchive& Ar)
+{
+	PRAGMA_DISABLE_DEPRECATION_WARNINGS
+	// When we remove the deprecated function, we can inline it here
+	SerializeCompressedData(Ar);
+	PRAGMA_ENABLE_DEPRECATION_WARNINGS
+}
+
 #if WITH_EDITOR
 struct FAnimDDCDebugData
 {
@@ -1383,7 +1393,7 @@ void FCompressedAnimSequence::SerializeCompressedData(FArchive& Ar, bool bDDCDat
 		if (BoneCompressionCodec != nullptr)
 		{
 			CompressedDataStructure = BoneCompressionCodec->AllocateAnimData();
-			CompressedDataStructure->SerializeCompressedData(Ar);
+			CompressedDataStructure->SerializeCompressedData(DataOwner, Ar);
 			CompressedDataStructure->Bind(CompressedByteStream);
 
 			// The codec can be null if we are a default object, a sequence with no raw bone data (just curves),
@@ -1500,7 +1510,7 @@ void FCompressedAnimSequence::SerializeCompressedData(FArchive& Ar, bool bDDCDat
 
 		if (BoneCompressionCodec != nullptr)
 		{
-			CompressedDataStructure->SerializeCompressedData(Ar);
+			CompressedDataStructure->SerializeCompressedData(DataOwner, Ar);
 		}
 	}
 
diff --git a/Engine/Source/Runtime/Engine/Public/Animation/AnimCompressionTypes.h b/Engine/Source/Runtime/Engine/Public/Animation/AnimCompressionTypes.h
--- a/Engine/Source/Runtime/Engine/Public/Animation/AnimCompressionTypes.h
+++ b/Engine/Source/Runtime/Engine/Public/Animation/AnimCompressionTypes.h
@@ -450,7 +450,10 @@ struct ICompressedAnimData
 	/* Virtual interface codecs must implement */
 	virtual ~ICompressedAnimData() {}
 
+	UE_DEPRECATED(5.5, "SerializeCompressedData now accepts the DataOwner as an argument")
 	ENGINE_API virtual void SerializeCompressedData(class FArchive& Ar);
+
+	ENGINE_API virtual void SerializeCompressedData(UObject* DataOwner, class FArchive& Ar);
 	virtual void Bind(const TArrayView<uint8> BulkData) = 0;
 
 	virtual int64 GetApproxCompressedSize() const = 0;
