Make sure cpp version matches ACLPlugin since we depend on it
Remove old Linux stdlib compatibility fix

diff --git a/Engine/Plugins/Animation/ACLPlugin/Source/ACLPlugin/ACLPlugin.Build.cs b/Engine/Plugins/Animation/ACLPlugin/Source/ACLPlugin/ACLPlugin.Build.cs
--- a/Engine/Plugins/Animation/ACLPlugin/Source/ACLPlugin/ACLPlugin.Build.cs
+++ b/Engine/Plugins/Animation/ACLPlugin/Source/ACLPlugin/ACLPlugin.Build.cs
@@ -11,7 +11,7 @@ namespace UnrealBuildTool.Rules
 		{
 			CppStandard = CppStandardVersion.Cpp17;
 
-            string ACLSDKDir = Path.GetFullPath(Path.Combine(ModuleDirectory, "../ThirdParty"));
+			string ACLSDKDir = Path.GetFullPath(Path.Combine(ModuleDirectory, "../ThirdParty"));
 
 			// Replace with PCHUsageMode.UseExplicitOrSharedPCHs when this plugin can compile with cpp20
 			PCHUsage = PCHUsageMode.NoPCHs;
@@ -28,15 +28,6 @@ namespace UnrealBuildTool.Rules
 				PrivateDependencyModuleNames.Add("DesktopPlatform");
 				PrivateDependencyModuleNames.Add("UnrealEd");
 			}
-
-			if (Target.Platform == UnrealTargetPlatform.Linux)
-			{
-				// There appears to be a bug when cross-compiling Linux under Windows where the clang tool-chain used
-				// isn't fully C++11 compliant. The standard specifies that when the 'cinttypes' header is included
-				// the format macros are always defined unlike C which requires the following macro to be defined first.
-				// This fix should be required for UE 4.20 and earlier versions.
-				PrivateDefinitions.Add("__STDC_FORMAT_MACROS");
-			}
 		}
 	}
 }
