Ensure ACL curve compression never fails to gracefully handle bad data

diff --git a/Engine/Plugins/Animation/ACLPlugin/Source/ACLPlugin/Classes/AnimCurveCompressionCodec_ACL.h b/Engine/Plugins/Animation/ACLPlugin/Source/ACLPlugin/Classes/AnimCurveCompressionCodec_ACL.h
--- a/Engine/Plugins/Animation/ACLPlugin/Source/ACLPlugin/Classes/AnimCurveCompressionCodec_ACL.h
+++ b/Engine/Plugins/Animation/ACLPlugin/Source/ACLPlugin/Classes/AnimCurveCompressionCodec_ACL.h
@@ -36,6 +36,7 @@ class UAnimCurveCompressionCodec_ACL : public UAnimCurveCompressionCodec
 #endif
 
 	// UAnimCurveCompressionCodec implementation
+	virtual bool ValidateCompressedData(UObject* DataOwner, const FCompressedAnimSequence& AnimSeq) const override;
 	virtual void DecompressCurves(const FCompressedAnimSequence& AnimSeq, FBlendedCurve& Curves, float CurrentTime) const override;
 	virtual float DecompressCurve(const FCompressedAnimSequence& AnimSeq, FName CurveName, float CurrentTime) const override;
 };
diff --git a/Engine/Plugins/Animation/ACLPlugin/Source/ACLPlugin/Private/AnimCurveCompressionCodec_ACL.cpp b/Engine/Plugins/Animation/ACLPlugin/Source/ACLPlugin/Private/AnimCurveCompressionCodec_ACL.cpp
--- a/Engine/Plugins/Animation/ACLPlugin/Source/ACLPlugin/Private/AnimCurveCompressionCodec_ACL.cpp
+++ b/Engine/Plugins/Animation/ACLPlugin/Source/ACLPlugin/Private/AnimCurveCompressionCodec_ACL.cpp
@@ -54,7 +54,7 @@ void UAnimCurveCompressionCodec_ACL::PopulateDDCKey(FArchive& Ar)
 		}
 	}
 
-	uint32 ForceRebuildVersion = 2;
+	uint32 ForceRebuildVersion = 3;
 	Ar << ForceRebuildVersion;
 
 	uint16 LatestACLVersion = static_cast<uint16>(acl::compressed_tracks_version16::latest);
@@ -103,6 +103,29 @@ static TArray<float> GetMorphTargetMaxPositionDeltas(const FCompressibleAnimData
 	return MorphTargetMaxPositionDeltas;
 }
 
+static void ResetTracksToIdentity(acl::track_array_float1f& ACLTracks)
+{
+	// This resets the input ACL track samples to the identity value but retains all other values
+	const uint32 NumSamples = 1;
+	const float SampleRate = 30.0f;
+
+	const float IdentityValue = 0.0f;
+
+	const acl::track_desc_scalarf DefaultDesc;
+
+	for (acl::track_float1f& ACLTrack : ACLTracks)
+	{
+		// Reset everything to the identity value and default values
+		// Retain the output index to ensure proper output size
+		acl::track_desc_scalarf Desc = ACLTrack.get_description();	// Copy
+		Desc.precision = DefaultDesc.precision;
+
+		// Reset track to a single sample
+		ACLTrack = acl::track_float1f::make_reserve(Desc, ACLAllocatorImpl, NumSamples, SampleRate);
+		ACLTrack[0] = IdentityValue;
+	}
+}
+
 bool UAnimCurveCompressionCodec_ACL::Compress(const FCompressibleAnimData& AnimSeq, FAnimCurveCompressionResult& OutResult)
 {
 	const TArray<float> MorphTargetMaxPositionDeltas = GetMorphTargetMaxPositionDeltas(AnimSeq, MorphTargetSource);
@@ -179,25 +202,61 @@ bool UAnimCurveCompressionCodec_ACL::Compress(const FCompressibleAnimData& AnimS
 
 	acl::compressed_tracks* CompressedTracks = nullptr;
 	acl::output_stats Stats;
-	const acl::error_result CompressionResult = acl::compress_track_list(ACLAllocatorImpl, Tracks, Settings, CompressedTracks, Stats);
+	acl::error_result CompressionResult = acl::compress_track_list(ACLAllocatorImpl, Tracks, Settings, CompressedTracks, Stats);
+
+	bool bEnableErrorReporting = true;
+	bool bCompressionFailed = false;
 
 	if (CompressionResult.any())
 	{
-		UE_LOG(LogAnimationCompression, Warning, TEXT("ACL failed to compress curves: %s [%s]"), ANSI_TO_TCHAR(CompressionResult.c_str()), *AnimSeq.FullName);
-		return false;
+		// If compression failed, one of two things happened:
+		//    * Invalid settings were used, this would be a code/logic error that results in an improper usage of ACL
+		//    * Invalid data was provided, this would be a validation error that should ideally be caught earlier (e.g import, save)
+		// 
+		// Either way, if we get here, we cannot recover and we cannot fail as the engine assumes that compression always succeeds.
+		// We must handle failure gracefully. To that end, we compress an empty stub to ensure that something is present to
+		// decompress. Because the stub is empty, we'll simply output the default values. We still log this as an error to signal that
+		// this is a problem that needs to be fixed. This will allow the editor to continue working with the default values we'll output
+		// but cooking will fail preventing us from running with invalid state.
+
+		UE_LOG(LogAnimationCompression, Error, TEXT("ACL failed to compress curves: %s [%s]"), ANSI_TO_TCHAR(CompressionResult.c_str()), *AnimSeq.FullName);
+
+		// We reset the tracks to the identity, getting rid of any potentially invalid data.
+		ResetTracksToIdentity(Tracks);
+
+		CompressionResult = acl::compress_track_list(ACLAllocatorImpl, Tracks, Settings, CompressedTracks, Stats);
+
+		// The stub compression should never fail
+		check(CompressionResult.empty() && CompressedTracks != nullptr);
+
+		// Because we compress an empty stub, disable error reporting below
+		bEnableErrorReporting = false;
+		bCompressionFailed = true;
 	}
 
 	checkSlow(CompressedTracks->is_valid(true).empty());
 
 	const uint32 CompressedDataSize = CompressedTracks->get_size();
 
-	OutResult.CompressedBytes.Empty(CompressedDataSize);
-	OutResult.CompressedBytes.AddUninitialized(CompressedDataSize);
+	// When compression fails, we add an extra few bytes of padding at the end
+	// This allows us to detect that the size is different so that we can output an error when validating
+	const uint32 ErrorPaddingValue = 0xFAFACDCD;
+	const uint32 ErrorPaddingSize = bCompressionFailed ? sizeof(ErrorPaddingValue) : 0;
+
+	OutResult.CompressedBytes.Empty(CompressedDataSize + ErrorPaddingSize);
+	OutResult.CompressedBytes.AddUninitialized(CompressedDataSize + ErrorPaddingSize);
 	FMemory::Memcpy(OutResult.CompressedBytes.GetData(), CompressedTracks, CompressedDataSize);
 
+	if (bCompressionFailed)
+	{
+		// Ensure our padding is deterministic (might not be aligned)
+		FMemory::Memcpy(OutResult.CompressedBytes.GetData() + CompressedDataSize, &ErrorPaddingValue, sizeof(ErrorPaddingValue));
+	}
+
 	OutResult.Codec = this;
 
 #if !NO_LOGGING
+	if (bEnableErrorReporting)
 	{
 		acl::decompression_context<acl::debug_scalar_decompression_settings> Context;
 		Context.initialize(*CompressedTracks);
@@ -213,6 +272,39 @@ bool UAnimCurveCompressionCodec_ACL::Compress(const FCompressibleAnimData& AnimS
 }
 #endif // WITH_EDITORONLY_DATA
 
+bool UAnimCurveCompressionCodec_ACL::ValidateCompressedData(UObject* DataOwner, const FCompressedAnimSequence& AnimSeq) const
+{
+	if (AnimSeq.IndexedCurveNames.Num() == 0)
+	{
+		return true;
+	}
+
+	const acl::compressed_tracks* CompressedTracks = acl::make_compressed_tracks(AnimSeq.CompressedCurveByteStream.GetData());
+	if (CompressedTracks == nullptr || CompressedTracks->is_valid(false).any())
+	{
+		UE_LOG(LogAnimationCompression, Error,
+			TEXT("ACL compressed curve data is missing or corrupted for an anim sequence: %s"),
+			DataOwner != nullptr ? *DataOwner->GetPathName() : TEXT("[Unknown Sequence]"));
+		return false;
+	}
+
+	// Check if we have padding and if it has our magic value to signal failure (might not be aligned)
+	const uint32 CompressedSize = CompressedTracks->get_size();
+	const uint32 CompressedCurveByteStreamSize = AnimSeq.CompressedCurveByteStream.Num();
+	const uint32 ErrorPaddingValue = (CompressedSize + sizeof(uint32)) <= CompressedCurveByteStreamSize ?
+		*reinterpret_cast<const uint32*>(&AnimSeq.CompressedCurveByteStream[CompressedSize]) : 0;
+	if (ErrorPaddingValue == 0xFAFACDCD)
+	{
+		UE_LOG(LogAnimationCompression, Error,
+			TEXT("ACL failed to compress curves for an anim sequence and will output the default curve values at runtime: %s"),
+			DataOwner != nullptr ? *DataOwner->GetPathName() : TEXT("[Unknown Sequence]"));
+		return false;
+	}
+
+	// All good!
+	return true;
+}
+
 struct UECurveDecompressionSettings final : public acl::decompression_settings
 {
 	static constexpr bool is_track_type_supported(acl::track_type8 type) { return type == acl::track_type8::float1f; }
diff --git a/Engine/Source/Runtime/Engine/Classes/Animation/AnimCurveCompressionCodec.h b/Engine/Source/Runtime/Engine/Classes/Animation/AnimCurveCompressionCodec.h
--- a/Engine/Source/Runtime/Engine/Classes/Animation/AnimCurveCompressionCodec.h
+++ b/Engine/Source/Runtime/Engine/Classes/Animation/AnimCurveCompressionCodec.h
@@ -41,6 +41,15 @@ class UAnimCurveCompressionCodec : public UObject
 
 	ENGINE_API virtual void Serialize(FArchive& Ar) override;
 
+	/*
+	 * Called on load and when cooking to validate that the compressed data is good.
+	 * Codecs should perform necessary validation and emit an error when appropriate.
+	 * e.g: UE_LOG(LogAnimationCompression, Error, TEXT("Bad data!"));
+	 * 
+	 * Returns true when the data is valid, false otherwise.
+	 */
+	ENGINE_API virtual bool ValidateCompressedData(UObject* DataOwner, const FCompressedAnimSequence& AnimSeq) const { return true; }
+
 #if WITH_EDITORONLY_DATA
 	/** Returns whether or not we can use this codec to compress. */
 	virtual bool IsCodecValid() const { return true; }
diff --git a/Engine/Source/Runtime/Engine/Private/Animation/AnimCompressionTypes.cpp b/Engine/Source/Runtime/Engine/Private/Animation/AnimCompressionTypes.cpp
--- a/Engine/Source/Runtime/Engine/Private/Animation/AnimCompressionTypes.cpp
+++ b/Engine/Source/Runtime/Engine/Private/Animation/AnimCompressionTypes.cpp
@@ -1523,6 +1523,14 @@ void FCompressedAnimSequence::SerializeCompressedData(FArchive& Ar, bool bDDCDat
 		}
 	}
 
+	if (Ar.IsLoading() || Ar.IsCooking())
+	{
+		if (CurveCompressionCodec != nullptr)
+		{
+			CurveCompressionCodec->ValidateCompressedData(DataOwner, *this);
+		}
+	}
+
 #if WITH_EDITOR
 	if (bDDCData)
 	{
