Fix missing include in non-editor configurations

diff --git a/Engine/Plugins/Animation/ACLPlugin/Source/ACLPlugin/Private/AnimCurveCompressionCodec_ACL.cpp b/Engine/Plugins/Animation/ACLPlugin/Source/ACLPlugin/Private/AnimCurveCompressionCodec_ACL.cpp
--- a/Engine/Plugins/Animation/ACLPlugin/Source/ACLPlugin/Private/AnimCurveCompressionCodec_ACL.cpp
+++ b/Engine/Plugins/Animation/ACLPlugin/Source/ACLPlugin/Private/AnimCurveCompressionCodec_ACL.cpp
@@ -7,6 +7,7 @@
 
 #include UE_INLINE_GENERATED_CPP_BY_NAME(AnimCurveCompressionCodec_ACL)
 
+#include "AnimationCompression.h"
 #include "Animation/AnimCurveUtils.h"
 #include "AnimationCompression.h"
 
