[AutoRTFM] Use FMemory::Malloc instead of GMalloc in ACLPlugin.

GMalloc is the low-level allocator. It is not transactionally safe and doesn't interact with LLM stats. Higher level code should prefer FMemory::Malloc instead.

diff --git a/Engine/Plugins/Animation/ACLPlugin/Source/ACLPlugin/Public/ACLImpl.h b/Engine/Plugins/Animation/ACLPlugin/Source/ACLPlugin/Public/ACLImpl.h
--- a/Engine/Plugins/Animation/ACLPlugin/Source/ACLPlugin/Public/ACLImpl.h
+++ b/Engine/Plugins/Animation/ACLPlugin/Source/ACLPlugin/Public/ACLImpl.h
@@ -4,6 +4,7 @@
 // Copyright 2018 Nicholas Frechette. All Rights Reserved.
 
 #include "CoreMinimal.h"
+#include "HAL/UnrealMemory.h"
 #include "UObject/ObjectMacros.h"
 #include "Runtime/Launch/Resources/Version.h"
 
@@ -99,12 +100,12 @@ class ACLAllocator final : public acl::iallocator
 public:
 	virtual void* allocate(size_t size, size_t alignment = acl::iallocator::k_default_alignment)
 	{
-		return GMalloc->Malloc(size, alignment);
+		return FMemory::Malloc(size, alignment);
 	}
 
 	virtual void deallocate(void* ptr, size_t size)
 	{
-		GMalloc->Free(ptr);
+		FMemory::Free(ptr);
 	}
 };
 
diff --git a/Engine/Plugins/Animation/ACLPlugin/Source/ACLPluginEditor/Private/ACLStatsDumpCommandlet.cpp b/Engine/Plugins/Animation/ACLPlugin/Source/ACLPluginEditor/Private/ACLStatsDumpCommandlet.cpp
--- a/Engine/Plugins/Animation/ACLPlugin/Source/ACLPluginEditor/Private/ACLStatsDumpCommandlet.cpp
+++ b/Engine/Plugins/Animation/ACLPlugin/Source/ACLPluginEditor/Private/ACLStatsDumpCommandlet.cpp
@@ -5,6 +5,7 @@
 
 #include "HAL/FileManagerGeneric.h"
 #include "HAL/PlatformTime.h"
+#include "HAL/UnrealMemory.h"
 #include "Runtime/CoreUObject/Public/UObject/UObjectIterator.h"
 #include "Runtime/Engine/Classes/Animation/AnimBoneCompressionSettings.h"
 #include "Runtime/Engine/Classes/Animation/AnimCompress.h"
@@ -87,7 +88,7 @@ static const TCHAR* ReadACLClip(FFileManagerGeneric& FileManager, const FString&
 
 	// Allocate directly without a TArray to automatically manage the memory because some
 	// clips are larger than 2 GB
-	char* RawData = static_cast<char*>(GMalloc->Malloc(Size));
+	char* RawData = static_cast<char*>(FMemory::Malloc(Size));
 
 	Reader->Serialize(RawData, Size);
 	Reader->Close();
@@ -97,14 +98,14 @@ static const TCHAR* ReadACLClip(FFileManagerGeneric& FileManager, const FString&
 		acl::compressed_tracks* CompressedTracks = reinterpret_cast<acl::compressed_tracks*>(RawData);
 		if (Size != CompressedTracks->get_size() || CompressedTracks->is_valid(true).any())
 		{
-			GMalloc->Free(RawData);
+			FMemory::Free(RawData);
 			return TEXT("Invalid binary ACL file provided");
 		}
 
 		const acl::error_result Result = acl::convert_track_list(Allocator, *CompressedTracks, OutTracks);
 		if (Result.any())
 		{
-			GMalloc->Free(RawData);
+			FMemory::Free(RawData);
 			return TEXT("Failed to convert input binary track list");
 		}
 	}
@@ -114,21 +115,21 @@ static const TCHAR* ReadACLClip(FFileManagerGeneric& FileManager, const FString&
 
 		if (ClipReader.get_file_type() != acl::sjson_file_type::raw_clip)
 		{
-			GMalloc->Free(RawData);
+			FMemory::Free(RawData);
 			return TEXT("SJSON file isn't a raw clip");
 		}
 
 		acl::sjson_raw_clip RawClip;
 		if (!ClipReader.read_raw_clip(RawClip))
 		{
-			GMalloc->Free(RawData);
+			FMemory::Free(RawData);
 			return TEXT("Failed to read ACL raw clip from file");
 		}
 
 		OutTracks = MoveTemp(RawClip.track_list);
 	}
 
-	GMalloc->Free(RawData);
+	FMemory::Free(RawData);
 	return nullptr;
 }
 
