Move ACL related natvis files to a better location

diff --git a/Engine/Plugins/Animation/ACLPlugin/Source/ACLPlugin/Private/AnimCurveCompressionCodec_ACL.cpp b/Engine/Plugins/Animation/ACLPlugin/Source/ACLPlugin/Private/AnimCurveCompressionCodec_ACL.cpp
--- a/Engine/Plugins/Animation/ACLPlugin/Source/ACLPlugin/Private/AnimCurveCompressionCodec_ACL.cpp
+++ b/Engine/Plugins/Animation/ACLPlugin/Source/ACLPlugin/Private/AnimCurveCompressionCodec_ACL.cpp
@@ -171,7 +171,7 @@ bool UAnimCurveCompressionCodec_ACL::Compress(const FCompressibleAnimData& AnimS
 		// If the delta is zero, any error is irrelevant as it will have no measurable impact.
 		// By dividing the precision value we want with the delta length, we can control how much precision our blend weight needs.
 		// If we want 0.01 cm precision and our largest vertex displacement is 3 cm, the blend weight precision needs to be:
-		// 0.01 cm / 3.00 cm = 0.0033 (with the units cancelling out just like we need)
+		// 0.01 cm / 3.00 cm = 0.0033 (with the units canceling out just like we need)
 		// Another way to think about it is that every 0.0033 increment of the blend weight results in an increment of 0.01 cm
 		// when our displacement delta is 3 cm.
 		// 0.01 cm / 50.00 cm = 0.0002 (if our delta increases, we need to retain more blend weight precision)
